Added option so cancelled nodes in a flow queue do not cancel the parent flow

diff --git a/Engine/Plugins/Experimental/ControlFlows/Source/ControlFlows/Private/ControlFlow.cpp b/Engine/Plugins/Experimental/ControlFlows/Source/ControlFlows/Private/ControlFlow.cpp
--- a/Engine/Plugins/Experimental/ControlFlows/Source/ControlFlows/Private/ControlFlow.cpp
+++ b/Engine/Plugins/Experimental/ControlFlows/Source/ControlFlows/Private/ControlFlow.cpp
@@ -58,7 +58,7 @@ void FControlFlow::HandleControlFlowNodeCompleted(TSharedRef<const FControlFlowN
 
 	if (ensure(CurrentNode.IsValid() && &NodeCompleted.Get() == CurrentNode.Get()))
 	{
-		const bool bCancelRequested = NodeCompleted->HasCancelBeenRequested();
+		const bool bCancelRequested = !bInterpretCancelledNodeAsComplete && NodeCompleted->HasCancelBeenRequested();
 		CurrentNode.Reset();
 		CurrentlyRunningTask.Reset();
 
@@ -210,6 +210,12 @@ void FControlFlow::CancelFlow()
 	}
 }
 
+FControlFlow& FControlFlow::SetCancelledNodeAsComplete(bool bCancelledNodeIsComplete)
+{
+	bInterpretCancelledNodeAsComplete = bCancelledNodeIsComplete;
+	return *this;
+}
+
 FString FControlFlow::FormatOrGetNewNodeDebugName(const FString& FlowNodeDebugName)
 {
 	if (FlowNodeDebugName.IsEmpty())
diff --git a/Engine/Plugins/Experimental/ControlFlows/Source/ControlFlows/Public/ControlFlow.h b/Engine/Plugins/Experimental/ControlFlows/Source/ControlFlows/Public/ControlFlow.h
--- a/Engine/Plugins/Experimental/ControlFlows/Source/ControlFlows/Public/ControlFlow.h
+++ b/Engine/Plugins/Experimental/ControlFlows/Source/ControlFlows/Public/ControlFlow.h
@@ -63,6 +63,8 @@ public:
 	/** Will cancel ALL flows, both child ControlFlows and ControlFlows who owns this Flow. You've been warned. */
 	void CancelFlow();
 
+	FControlFlow& SetCancelledNodeAsComplete(bool bCancelledNodeIsComplete);
+
 public:
 
 	template<typename...ArgsT>
@@ -250,6 +252,8 @@ private:
 
 	FString DebugName;
 
+	bool bInterpretCancelledNodeAsComplete = false;
+
 	int32 UnnamedNodeCounter = 0;
 
 	int32 UnnamedBranchCounter = 0;
