ControlFlow - Add an accessor to get the debug name of the current step you're performing.

diff --git a/Engine/Plugins/Experimental/ControlFlows/Source/ControlFlows/Private/ControlFlow.cpp b/Engine/Plugins/Experimental/ControlFlows/Source/ControlFlows/Private/ControlFlow.cpp
--- a/Engine/Plugins/Experimental/ControlFlows/Source/ControlFlows/Private/ControlFlow.cpp
+++ b/Engine/Plugins/Experimental/ControlFlows/Source/ControlFlows/Private/ControlFlow.cpp
@@ -227,6 +227,21 @@ FString FControlFlow::FormatOrGetNewNodeDebugName(const FString& FlowNodeDebugNa
 	return FlowNodeDebugName;
 }
 
+TOptional<FString> FControlFlow::GetCurrentStepDebugName() const
+{
+	if (CurrentNode.IsValid())
+	{
+		if (CurrentlyRunningTask.IsValid())
+		{
+			return CurrentlyRunningTask->GetNodeName();
+		}
+		
+		return CurrentNode->GetNodeName();
+	}
+
+	return TOptional<FString>();
+}
+
 FSimpleDelegate& FControlFlow::QueueFunction(const FString& FlowNodeDebugName)
 {
 	TSharedRef<FControlFlowNode_SelfCompleting> NewNode = MakeShared<FControlFlowNode_SelfCompleting>(SharedThis(this), FormatOrGetNewNodeDebugName(FlowNodeDebugName));
@@ -341,4 +356,4 @@ void FControlFlow::HandleOnTaskCancelled()
 	{
 		CurrentlyRunningTask->CompleteCancelFlow();
 	}
-}
\ No newline at end of file
+}
diff --git a/Engine/Plugins/Experimental/ControlFlows/Source/ControlFlows/Public/ControlFlow.h b/Engine/Plugins/Experimental/ControlFlows/Source/ControlFlows/Public/ControlFlow.h
--- a/Engine/Plugins/Experimental/ControlFlows/Source/ControlFlows/Public/ControlFlow.h
+++ b/Engine/Plugins/Experimental/ControlFlows/Source/ControlFlows/Public/ControlFlow.h
@@ -65,6 +65,8 @@ public:
 
 	FControlFlow& SetCancelledNodeAsComplete(bool bCancelledNodeIsComplete);
 
+	TOptional<FString> GetCurrentStepDebugName() const;
+
 public:
 
 	template<typename...ArgsT>
