Added new ensure messaging to help guide people on what to do.

diff --git a/Engine/Plugins/Experimental/ControlFlows/Source/ControlFlows/Public/ControlFlowManager.h b/Engine/Plugins/Experimental/ControlFlows/Source/ControlFlows/Public/ControlFlowManager.h
--- a/Engine/Plugins/Experimental/ControlFlows/Source/ControlFlows/Public/ControlFlowManager.h
+++ b/Engine/Plugins/Experimental/ControlFlows/Source/ControlFlows/Public/ControlFlowManager.h
@@ -16,7 +16,7 @@ public:
 		ensureAlwaysMsgf(FlowId.Len() > 0, TEXT("All Flows need a non-empty ID!"));
 
 		FControlFlow* AlreadyExistedFlow = Find(OwningObject, FlowId);
-		if (!ensureAlwaysMsgf(!AlreadyExistedFlow, TEXT("Flow already exists! All flows should have a unique ID. If there are multiple instances of the owning object, that might cause this!")))
+		if (!ensureAlwaysMsgf(!AlreadyExistedFlow, TEXT("Flow already exists! All flows should have a unique ID. If there are multiple instances of the owning object, that might cause this! You can call MakeShared<FControlFlow> instead of using FControlFlowStatics!")))
 		{
 			AlreadyExistedFlow->Reset();
 		}
