ControlFlow Expansion: ForkFlow

You can chain your flow with this lambda:

.ForkFlow([this](TSharedRef<FConcurrentControlFlows> ConcurrentFlows) {})

And that will allow you create multiple control flows to execute "at the same time" (still single-threaded).

The default continue behavior from the Fork is to wait until all concurrent flows have been [completed or cancelled]

diff --git a/Engine/Plugins/Experimental/ControlFlows/Source/ControlFlows/Private/ControlFlow.cpp b/Engine/Plugins/Experimental/ControlFlows/Source/ControlFlows/Private/ControlFlow.cpp
--- a/Engine/Plugins/Experimental/ControlFlows/Source/ControlFlows/Private/ControlFlow.cpp
+++ b/Engine/Plugins/Experimental/ControlFlows/Source/ControlFlows/Private/ControlFlow.cpp
@@ -284,6 +284,19 @@ FControlFlowBranchDefiner& FControlFlow::QueueControlFlowBranch(const FString& T
 	return NewTask->GetDelegate();
 }
 
+FConcurrentFlowsDefiner& FControlFlow::QueueConcurrentFlows(const FString& TaskName /*= TEXT("")*/, const FString& FlowNodeDebugName /*= TEXT("")*/)
+{
+	TSharedRef<FControlFlowTask_ConcurrentFlows> NewTask = MakeShared<FControlFlowTask_ConcurrentFlows>(TaskName);
+	TSharedRef<FControlFlowNode_Task> NewNode = MakeShared<FControlFlowNode_Task>(SharedThis(this), NewTask, FormatOrGetNewNodeDebugName(FlowNodeDebugName));
+
+	NewNode->OnExecute().BindSP(SharedThis(this), &FControlFlow::HandleTaskNodeExecuted);
+	NewNode->OnCancelRequested().BindSP(SharedThis(this), &FControlFlow::HandleTaskNodeCancelled);
+
+	FlowQueue.Add(NewNode);
+
+	return NewTask->GetDelegate();
+}
+
 TSharedRef<FControlFlowTask_BranchLegacy> FControlFlow::QueueBranch(FControlFlowBranchDecider_Legacy& BranchDecider, const FString& TaskName /*= TEXT("")*/, const FString& FlowNodeDebugName /*= TEXT("")*/)
 {
 	ensureAlwaysMsgf(false, TEXT("Deprecated. Use 'QueueControlFlowBranch'"));
diff --git a/Engine/Plugins/Experimental/ControlFlows/Source/ControlFlows/Private/ControlFlowConcurrency.cpp b/Engine/Plugins/Experimental/ControlFlows/Source/ControlFlows/Private/ControlFlowConcurrency.cpp
new file mode 100644
--- /dev/null
+++ b/Engine/Plugins/Experimental/ControlFlows/Source/ControlFlows/Private/ControlFlowConcurrency.cpp
@@ -0,0 +1,233 @@
+// Copyright Epic Games, Inc. All Rights Reserved.
+
+#include "ControlFlowConcurrency.h"
+#include "ControlFlow.h"
+#include "ControlFlows.h"
+
+FConcurrencySubFlowContainer::FConcurrencySubFlowContainer(const FString& InDebugName)
+	: SubFlow(MakeShared<FControlFlow>(InDebugName))
+{}
+
+bool FConcurrencySubFlowContainer::HasBeenExecuted() const
+{
+	return bHasBeenExecuted;
+}
+
+bool FConcurrencySubFlowContainer::HasBeenCancelled() const
+{
+	return bHasBeenCancelled;
+}
+
+bool FConcurrencySubFlowContainer::IsCompleteOrCancelled() const
+{
+	const bool bIsComplete = HasBeenExecuted() && !SubFlow->IsRunning();
+	return bIsComplete || HasBeenCancelled();
+}
+
+void FConcurrencySubFlowContainer::Execute()
+{
+	UE_LOG(LogControlFlows, Verbose, TEXT("ConcurrentControlFlowContainer - Executing %s"), *GetDebugName());
+
+	if (ensureAlways(!IsCompleteOrCancelled()))
+	{
+		bHasBeenExecuted = true;
+		SubFlow->ExecuteFlow();
+	}
+}
+
+void FConcurrencySubFlowContainer::Cancel()
+{
+	UE_LOG(LogControlFlows, Verbose, TEXT("ConcurrentControlFlowContainer - Cancelling %s"), *GetDebugName());
+
+	bHasBeenCancelled = true;
+	if (SubFlow->IsRunning())
+	{
+		SubFlow->CancelFlow();
+	}
+}
+
+FSimpleDelegate& FConcurrencySubFlowContainer::OnComplete()
+{
+	return SubFlow->OnComplete();
+}
+
+FSimpleDelegate& FConcurrencySubFlowContainer::OnExecutedWithoutAnyNodes()
+{
+	return SubFlow->OnExecutedWithoutAnyNodes();
+}
+
+FSimpleDelegate& FConcurrencySubFlowContainer::OnCancelled()
+{
+	return SubFlow->OnCancelled();
+}
+
+const FString& FConcurrencySubFlowContainer::GetDebugName()
+{
+	return SubFlow->GetDebugName();
+}
+
+TSharedRef<FControlFlow> FConcurrencySubFlowContainer::GetControlFlow() const
+{
+	return SubFlow;
+}
+
+////////////////////////////////////////////////////////////////////////////////////////
+
+FControlFlow& FConcurrentControlFlows::AddOrGetFlow(int32 InIdentifier, const FString& DebugSubFlowName /*= TEXT("")*/)
+{
+	if (TSharedRef<FConcurrencySubFlowContainer>* ExistingFlow = ConcurrentFlows.Find(InIdentifier))
+	{
+		return (*ExistingFlow)->GetControlFlow().Get();
+	}
+	else
+	{
+		static int32 UnnamedControlConcurrencyFlowCounter = 0;
+		const FString& DebugNameToUse = DebugSubFlowName.IsEmpty() ? FString::Format(TEXT("UnnamedConcurrencySubFlow_{0}"), { (UnnamedControlConcurrencyFlowCounter++) }) : DebugSubFlowName;
+
+		TSharedRef<FConcurrencySubFlowContainer> NewConcurrentFlow = MakeShared<FConcurrencySubFlowContainer>(DebugNameToUse);
+
+		NewConcurrentFlow->OnComplete().BindSP(SharedThis(this), &FConcurrentControlFlows::HandleConcurrentFlowCompleted, InIdentifier);
+		NewConcurrentFlow->OnExecutedWithoutAnyNodes().BindSP(SharedThis(this), &FConcurrentControlFlows::HandleConcurrentFlowCompleted, InIdentifier);
+		NewConcurrentFlow->OnCancelled().BindSP(SharedThis(this), &FConcurrentControlFlows::HandleConcurrentFlowCancelled, InIdentifier);
+
+		ConcurrentFlows.Add(InIdentifier, NewConcurrentFlow);
+		return NewConcurrentFlow->GetControlFlow().Get();
+	}
+}
+
+FControlFlow& FConcurrentControlFlows::AddOrGetProng(int32 InIdentifier, const FString& DebugSubFlowName /*= TEXT("")*/)
+{
+	return AddOrGetFlow(InIdentifier, DebugSubFlowName);
+}
+
+bool FConcurrentControlFlows::AreAllSubFlowsCompletedOrCancelled() const
+{
+	for (const TPair<int32, TSharedRef<FConcurrencySubFlowContainer>>& PairIt : ConcurrentFlows)
+	{
+		if (!PairIt.Value->IsCompleteOrCancelled())
+		{
+			return false;
+		}
+	}
+
+	return true;
+}
+
+bool FConcurrentControlFlows::HasAnySubFlowBeenExecuted() const
+{
+	for (const TPair<int32, TSharedRef<FConcurrencySubFlowContainer>>& PairIt : ConcurrentFlows)
+	{
+		if (PairIt.Value->HasBeenExecuted())
+		{
+			return true;
+		}
+	}
+
+	return false;
+}
+
+void FConcurrentControlFlows::HandleConcurrentFlowCompleted(int32 FlowIndex)
+{
+	check(ConcurrentFlows.Contains(FlowIndex));
+
+	UE_LOG(LogControlFlows, Verbose, TEXT("ConcurrentControlFlow Finished: %s"), *ConcurrentFlows[FlowIndex]->GetDebugName());
+
+	if (!bCancelAllHasBegun)
+	{
+		if (GetConcurrencyBehavior().GetContinueCondition() == FConcurrentControlFlowBehavior::EContinueConditions::ContinueOnAll_CompletedOrCancelled)
+		{
+			if (AreAllSubFlowsCompletedOrCancelled())
+			{
+				OnAllCompleted();
+			}
+			else
+			{
+				UE_LOG(LogControlFlows, Verbose, TEXT("Other flows are still running"));
+			}
+		}
+		else
+		{
+			checkf(false, TEXT("Unhandled Continue Condition"));
+		}
+	}
+}
+
+void FConcurrentControlFlows::HandleConcurrentFlowCancelled(int32 FlowIndex)
+{
+	check(ConcurrentFlows.Contains(FlowIndex));
+
+	UE_LOG(LogControlFlows, Verbose, TEXT("ConcurrentControlFlow Cancelled %s"), *ConcurrentFlows[FlowIndex]->GetDebugName());
+
+	if (!bCancelAllHasBegun)
+	{
+		HandleConcurrentFlowCompleted(FlowIndex);
+	}
+}
+
+void FConcurrentControlFlows::Execute()
+{
+	ensureAlwaysMsgf(!this->HasAnySubFlowBeenExecuted(), TEXT("Did you call ExecuteFlow() on a SubFlow? Do not do this! You only need to call ExecuteFlow once per FControlFlowStatics::Create!"));
+
+	bool bFlowExecuted = false;
+
+	for (TPair<int32, TSharedRef<FConcurrencySubFlowContainer>>& PairIt : ConcurrentFlows)
+	{
+		UE_LOG(LogControlFlows, Verbose, TEXT("ConcurrentControlFlow::Execute - Executing Subflow %s"), *PairIt.Value->GetDebugName());
+
+		bFlowExecuted = true;
+		PairIt.Value->Execute();
+
+		if (bCancelAllHasBegun)
+		{
+			break;
+		}
+	}
+	
+	if (ConcurrentFlows.Num() > 0)
+	{
+		ensureAlwaysMsgf(bFlowExecuted, TEXT("There were defined Subflows, but none were executed"));
+	}
+
+	if (!bFlowExecuted)
+	{
+		OnAllCompleted();
+	}
+}
+
+void FConcurrentControlFlows::CancelAll()
+{
+	if (!bCancelAllHasBegun)
+	{
+		UE_LOG(LogControlFlows, Verbose, TEXT("FConcurrentControlFlows::CancelAll() called"));
+
+		bCancelAllHasBegun = true;
+
+		for (TPair<int32, TSharedRef<FConcurrencySubFlowContainer>>& PairIt : ConcurrentFlows)
+		{
+			UE_LOG(LogControlFlows, Verbose, TEXT("ConcurrentControlFlow::Cancel All - Cancelling Subflow %s"), *PairIt.Value->GetDebugName());
+
+			PairIt.Value->Cancel();
+		}
+
+		OnAllCancelled();
+	}
+}
+
+void FConcurrentControlFlows::OnAllCompleted()
+{
+	UE_LOG(LogControlFlows, Verbose, TEXT("ConcurrentControlFlow::OnAllCompleted()"));
+
+	ConcurrentFlows.Reset();
+	OnConcurrencyCompleted.ExecuteIfBound();
+}
+
+void FConcurrentControlFlows::OnAllCancelled()
+{
+	ConcurrentFlows.Reset();
+	OnConcurrencyCancelled.ExecuteIfBound();
+}
+
+FConcurrentControlFlowBehavior FConcurrentControlFlows::GetConcurrencyBehavior() const
+{
+	return FConcurrentControlFlowBehavior();
+}
diff --git a/Engine/Plugins/Experimental/ControlFlows/Source/ControlFlows/Private/ControlFlowTask.cpp b/Engine/Plugins/Experimental/ControlFlows/Source/ControlFlows/Private/ControlFlowTask.cpp
--- a/Engine/Plugins/Experimental/ControlFlows/Source/ControlFlows/Private/ControlFlowTask.cpp
+++ b/Engine/Plugins/Experimental/ControlFlows/Source/ControlFlows/Private/ControlFlowTask.cpp
@@ -3,6 +3,7 @@
 #include "ControlFlowTask.h"
 #include "ControlFlows.h"
 #include "ControlFlowBranch.h"
+#include "ControlFlowConcurrency.h"
 
 //////////////////////////
 //FControlFlowTask
@@ -306,3 +307,52 @@ void FControlFlowTask_Branch::HandleBranchCancelled()
 
 	OnCancelled().ExecuteIfBound();
 }
+
+//////////////////////////////////
+//FControlFlowTask_ConcurrentBranch
+//////////////////////////////////
+
+void FControlFlowTask_ConcurrentFlows::Execute()
+{
+	if (ensureAlways(ConcurrentFlowDelegate.IsBound() && !ConcurrentFlows.IsValid()))
+	{
+		ConcurrentFlows = MakeShared<FConcurrentControlFlows>();
+		ConcurrentFlows->OnConcurrencyCompleted.BindSP(SharedThis(this), &FControlFlowTask_ConcurrentFlows::HandleConcurrentFlowsCompleted);
+		ConcurrentFlows->OnConcurrencyCancelled.BindSP(SharedThis(this), &FControlFlowTask_ConcurrentFlows::HandleConcurrentFlowsCancelled);
+
+		ConcurrentFlowDelegate.Execute(ConcurrentFlows.ToSharedRef());
+		ConcurrentFlows->Execute();
+	}
+	else
+	{
+		HandleConcurrentFlowsCompleted();
+	}
+}
+
+void FControlFlowTask_ConcurrentFlows::Cancel()
+{
+	if (ensureAlwaysMsgf(ConcurrentFlows.IsValid(), TEXT("Did you call execute flow on a subflow? Do not do this! You only need to call ExecuteFlow once per FControlFlowStatics::Create!")))
+	{
+		ConcurrentFlows->CancelAll();
+	}
+}
+
+void FControlFlowTask_ConcurrentFlows::HandleConcurrentFlowsCompleted()
+{
+	if (ConcurrentFlowDelegate.IsBound())
+	{
+		ConcurrentFlowDelegate.Unbind();
+	}
+
+	OnComplete().ExecuteIfBound();
+}
+
+void FControlFlowTask_ConcurrentFlows::HandleConcurrentFlowsCancelled()
+{
+	if (ConcurrentFlowDelegate.IsBound())
+	{
+		ConcurrentFlowDelegate.Unbind();
+	}
+
+	OnCancelled().ExecuteIfBound();
+}
diff --git a/Engine/Plugins/Experimental/ControlFlows/Source/ControlFlows/Public/ControlFlow.h b/Engine/Plugins/Experimental/ControlFlows/Source/ControlFlows/Public/ControlFlow.h
--- a/Engine/Plugins/Experimental/ControlFlows/Source/ControlFlows/Public/ControlFlow.h
+++ b/Engine/Plugins/Experimental/ControlFlows/Source/ControlFlows/Public/ControlFlow.h
@@ -105,11 +105,19 @@ public:
 		return *this;
 	}
 
+	template<typename FunctionT, typename...ArgsT>
+	FControlFlow& ForkFlow(FunctionT InForkLambda, ArgsT...Params)
+	{
+		QueueConcurrentFlows(FormatOrGetNewNodeDebugName()).BindLambda(InForkLambda, Params...);
+		return *this;
+	}
+
 public:
 	FSimpleDelegate& QueueFunction(const FString& FlowNodeDebugName = TEXT(""));
 	FControlFlowWaitDelegate& QueueWait(const FString& FlowNodeDebugName = TEXT(""));
 	FControlFlowPopulator& QueueControlFlow(const FString& TaskName = TEXT(""), const FString& FlowNodeDebugName = TEXT(""));
 	FControlFlowBranchDefiner& QueueControlFlowBranch(const FString& TaskName = TEXT(""), const FString& FlowNodeDebugName = TEXT(""));
+	FConcurrentFlowsDefiner& QueueConcurrentFlows(const FString& TaskName = TEXT(""), const FString& FlowNodeDebugName = TEXT(""));
 
 private:
 	template<typename BindingObjectT, typename...PayloadParamsT>
@@ -141,6 +149,12 @@ private:
 
 private:
 	
+	template<typename BindingObjectClassT, typename...PayloadParamsT>
+	void QueueStep_Internal_TSharedFromThis(const FString& InDebugName, TSharedRef<BindingObjectClassT> InBindingObject, typename TMemFunPtrType<false, BindingObjectClassT, void(TSharedRef<FConcurrentControlFlows>, PayloadParamsT...)>::Type InFunction, PayloadParamsT...Params)
+	{
+		QueueConcurrentFlows(InDebugName).BindSP(InBindingObject, InFunction, Params...);
+	}
+
 	template<typename BindingObjectClassT, typename...PayloadParamsT>
 	void QueueStep_Internal_TSharedFromThis(const FString& InDebugName, TSharedRef<BindingObjectClassT> InBindingObject, typename TMemFunPtrType<false, BindingObjectClassT, int32(TSharedRef<FControlFlowBranch>, PayloadParamsT...)>::Type InFunction, PayloadParamsT...Params)
 	{
@@ -166,6 +180,13 @@ private:
 	}
 
 private:
+
+	template<typename BindingObjectClassT, typename...PayloadParamsT>
+	void QueueStep_Internal_UObject(const FString& InDebugName, BindingObjectClassT* InBindingObject, typename TMemFunPtrType<false, BindingObjectClassT, void(TSharedRef<FConcurrentControlFlows>, PayloadParamsT...)>::Type InFunction, PayloadParamsT...Params)
+	{
+		QueueConcurrentFlows(InDebugName).BindUObject(InBindingObject, InFunction, Params...);
+	}
+
 	template<typename BindingObjectClassT, typename...PayloadParamsT>
 	void QueueStep_Internal_UObject(const FString& InDebugName, BindingObjectClassT* InBindingObject, typename TMemFunPtrType<false, BindingObjectClassT, int32(TSharedRef<FControlFlowBranch>, PayloadParamsT...)>::Type InFunction, PayloadParamsT...Params)
 	{
@@ -198,6 +219,7 @@ private:
 	friend class FControlFlowTask_Loop;
 	friend class FControlFlowTask_BranchLegacy;
 	friend class FControlFlowTask_Branch;
+	friend struct FConcurrencySubFlowContainer;
 
 public:
 	/** These work, but they are a bit clunky to use. The heart of the issue is that it requires the caller to define two functions. We want only the caller to use one function.
diff --git a/Engine/Plugins/Experimental/ControlFlows/Source/ControlFlows/Public/ControlFlowBranch.h b/Engine/Plugins/Experimental/ControlFlows/Source/ControlFlows/Public/ControlFlowBranch.h
--- a/Engine/Plugins/Experimental/ControlFlows/Source/ControlFlows/Public/ControlFlowBranch.h
+++ b/Engine/Plugins/Experimental/ControlFlows/Source/ControlFlows/Public/ControlFlowBranch.h
@@ -11,7 +11,7 @@ static int32 UnnamedControlFlowBranchCounter = 0;
 class FControlFlowBranch : public TSharedFromThis<FControlFlowBranch>
 {
 public:
-	FControlFlow& AddOrGetBranch(const int32& InKey, const FString& BranchName = TEXT(""))
+	FControlFlow& AddOrGetBranch(const int32& InKey, const FString& DebugBranchName = TEXT(""))
 	{
 		if (TSharedRef<FControlFlow>* ExistingBranch = Branches.Find(InKey))
 		{
@@ -19,8 +19,8 @@ public:
 		}
 		else
 		{
-			const FString& NameToUse = BranchName.IsEmpty() ? FString::Format(TEXT("UnnamedBranch_{0}"), { (UnnamedControlFlowBranchCounter++) }) : BranchName;
-			TSharedRef<FControlFlow> NewControlFlow = MakeShared<FControlFlow>(NameToUse);
+			const FString& DebugNameToUse = DebugBranchName.IsEmpty() ? FString::Format(TEXT("UnnamedBranch_{0}"), { (UnnamedControlFlowBranchCounter++) }) : DebugBranchName;
+			TSharedRef<FControlFlow> NewControlFlow = MakeShared<FControlFlow>(DebugNameToUse);
 
 			Branches.Add(InKey, NewControlFlow);
 			return NewControlFlow.Get();
diff --git a/Engine/Plugins/Experimental/ControlFlows/Source/ControlFlows/Public/ControlFlowConcurrency.h b/Engine/Plugins/Experimental/ControlFlows/Source/ControlFlows/Public/ControlFlowConcurrency.h
new file mode 100644
--- /dev/null
+++ b/Engine/Plugins/Experimental/ControlFlows/Source/ControlFlows/Public/ControlFlowConcurrency.h
@@ -0,0 +1,100 @@
+// Copyright Epic Games, Inc. All Rights Reserved.
+
+#pragma once
+
+#include "Containers/Map.h"
+#include "Delegates/DelegateCombinations.h"
+#include "Templates/SharedPointer.h"
+
+class FControlFlow;
+class FConcurrentControlFlowBehavior;
+struct FConcurrencySubFlowContainer;
+
+// All Flows will be executed concurrently (Single-Threaded). WARNING: Having a non-terminating loop within a fork can cause a forever hang for the forked step
+class FConcurrentControlFlows : public TSharedFromThis<FConcurrentControlFlows>
+{
+public:
+	CONTROLFLOWS_API FControlFlow& AddOrGetProng(int32 InIdentifier, const FString& DebugSubFlowName = TEXT(""));
+	CONTROLFLOWS_API FControlFlow& AddOrGetFlow(int32 InIdentifier, const FString& DebugSubFlowName = TEXT(""));
+
+	// TODO: Multi-threaded-like behavior
+	// void Set(const FConcurrentControlFlowBehavior& InBehavior);
+
+private:
+	friend class FControlFlowTask_ConcurrentFlows;
+
+	bool AreAllSubFlowsCompletedOrCancelled() const;
+	bool HasAnySubFlowBeenExecuted() const;
+
+	void HandleConcurrentFlowCompleted(int32 FlowIndex);
+	void HandleConcurrentFlowCancelled(int32 FlowIndex);
+
+	void Execute();
+	void CancelAll(); // Do not make public
+
+	void OnAllCompleted();
+	void OnAllCancelled();
+
+	FSimpleDelegate OnConcurrencyCompleted;
+	FSimpleDelegate OnConcurrencyCancelled;
+
+	bool bCancelAllHasBegun = false;
+	TMap<int32, TSharedRef<FConcurrencySubFlowContainer>> ConcurrentFlows;
+
+private:
+	FConcurrentControlFlowBehavior GetConcurrencyBehavior() const;
+};
+
+struct FConcurrencySubFlowContainer : public TSharedFromThis<FConcurrencySubFlowContainer>
+{
+	FConcurrencySubFlowContainer(const FString& InDebugName);
+
+private:
+	friend class FConcurrentControlFlows;
+
+	bool HasBeenExecuted() const;
+	bool HasBeenCancelled() const;
+	bool IsCompleteOrCancelled() const;
+
+	void Execute();
+	void Cancel();
+	FSimpleDelegate& OnComplete();
+	FSimpleDelegate& OnExecutedWithoutAnyNodes();
+	FSimpleDelegate& OnCancelled();
+	const FString& GetDebugName();
+
+	TSharedRef<FControlFlow> GetControlFlow() const;
+
+	bool bHasBeenExecuted = false;
+	bool bHasBeenCancelled = false;
+
+	TSharedRef<FControlFlow> SubFlow;
+};
+
+// Placeholder class to extend Concurrency behavior. TODO: Read up more on multi-threaded processing
+class FConcurrentControlFlowBehavior
+{
+	friend class FConcurrentControlFlows;
+private:
+	enum class EContinueConditions
+	{
+		ContinueOnAll_CompletedOrCancelled, // Completed or Canceled
+		// Can add behavior here as needed before having an evaluation method
+	};
+
+	EContinueConditions GetContinueCondition() const { return EContinueConditions::ContinueOnAll_CompletedOrCancelled; }
+
+	// TODO: Customized continue conditions
+	// void EvaluateContinueCondition();
+	//
+	// Examples:
+	//		ContinueOnFirstCompleted_CancelOthers,
+	//		ContinueOnFirstCanceled_CancelOthers,
+	//		ContinueOnFirstCompletedOrCanceled_CancelOthers,
+	//		ContinueOnSpecificFlowsCompleted_CancelOthers,
+	//		ContinueOnSpecificFlowsCancelled_CancelOthers,
+	//		ContinueOnSpecificFlowsCompletedOrCancelled_CancelOthers
+
+	// TODO: Determine how a Forked Flow handles canceling
+	// Having a cancelled sub-flow affect other flows sort of goes against the idea of "Fork" flow, where each Subflow behaves independently from each other. Further debate required.
+};
\ No newline at end of file
diff --git a/Engine/Plugins/Experimental/ControlFlows/Source/ControlFlows/Public/ControlFlowNode.h b/Engine/Plugins/Experimental/ControlFlows/Source/ControlFlows/Public/ControlFlowNode.h
--- a/Engine/Plugins/Experimental/ControlFlows/Source/ControlFlows/Public/ControlFlowNode.h
+++ b/Engine/Plugins/Experimental/ControlFlows/Source/ControlFlows/Public/ControlFlowNode.h
@@ -10,6 +10,8 @@ class FControlFlowNode;
 class FControlFlowNode_Task;
 class FControlFlowSubTaskBase;
 class FControlFlowBranch;
+class FConcurrentControlFlows;
+class FConcurrentControlFlowBehavior;
 
 using FControlFlowNodePtr = TSharedPtr<FControlFlowNode>;
 using FControlFlowNodeRef = TSharedRef<FControlFlowNode>;
@@ -23,6 +25,8 @@ DECLARE_DELEGATE_RetVal(int32, FControlFlowBranchDecider_Legacy)
 // Would be nice to not force people to use int32 for the branch key, but finding the syntax to get it to work properly proved to be more difficult than it was worth. End result is callers have to cast to an int32.
 DECLARE_DELEGATE_RetVal_OneParam(int32, FControlFlowBranchDefiner, TSharedRef<FControlFlowBranch>);
 
+DECLARE_DELEGATE_OneParam(FConcurrentFlowsDefiner, TSharedRef<FConcurrentControlFlows>);
+
 class FControlFlowNode : public TSharedFromThis<FControlFlowNode>
 {
 public:
diff --git a/Engine/Plugins/Experimental/ControlFlows/Source/ControlFlows/Public/ControlFlowTask.h b/Engine/Plugins/Experimental/ControlFlows/Source/ControlFlows/Public/ControlFlowTask.h
--- a/Engine/Plugins/Experimental/ControlFlows/Source/ControlFlows/Public/ControlFlowTask.h
+++ b/Engine/Plugins/Experimental/ControlFlows/Source/ControlFlows/Public/ControlFlowTask.h
@@ -32,6 +32,8 @@ private:
 	FString DebugName;
 };
 
+//////////////////////////////////////////////////////////////////////////////
+
 class FControlFlowSimpleSubTask : public FControlFlowSubTaskBase
 {
 public:
@@ -56,6 +58,8 @@ private:
 	TSharedPtr<FControlFlow> TaskFlow = nullptr;
 };
 
+//////////////////////////////////////////////////////////////////////////////
+
 class FControlFlowTask_Loop : public FControlFlowSimpleSubTask
 {
 public:
@@ -74,6 +78,8 @@ private:
 	FControlFlowLoopComplete TaskCompleteDecider;
 };
 
+//////////////////////////////////////////////////////////////////////////////
+
 class FControlFlowTask_BranchLegacy : public FControlFlowSubTaskBase
 {
 public:
@@ -119,6 +125,8 @@ private:
 	int32 SelectedBranch = INDEX_NONE;
 };
 
+//////////////////////////////////////////////////////////////////////////////
+
 class FControlFlowTask_Branch : public FControlFlowSubTaskBase
 {
 public:
@@ -142,4 +150,31 @@ private:
 
 	FControlFlowBranchDefiner BranchDelegate;
 	TSharedPtr<FControlFlow> SelectedBranchFlow;
+};
+
+//////////////////////////////////////////////////////////////////////////////
+
+class FControlFlowTask_ConcurrentFlows : public FControlFlowSubTaskBase
+{
+public:
+	friend class FControlFlow;
+
+public:
+	FControlFlowTask_ConcurrentFlows(const FString& TaskName)
+		: FControlFlowSubTaskBase(TaskName)
+	{}
+
+	FConcurrentFlowsDefiner& GetDelegate() { return ConcurrentFlowDelegate; }
+
+protected:
+	virtual void Execute() override;
+	virtual void Cancel() override;
+
+private:
+
+	void HandleConcurrentFlowsCompleted();
+	void HandleConcurrentFlowsCancelled();
+
+	FConcurrentFlowsDefiner ConcurrentFlowDelegate;
+	TSharedPtr<FConcurrentControlFlows> ConcurrentFlows;
 };
\ No newline at end of file
