Fix

diff --git a/Engine/Plugins/Experimental/ControlFlows/Source/ControlFlows/Public/ControlFlowConcurrency.h b/Engine/Plugins/Experimental/ControlFlows/Source/ControlFlows/Public/ControlFlowConcurrency.h
--- a/Engine/Plugins/Experimental/ControlFlows/Source/ControlFlows/Public/ControlFlowConcurrency.h
+++ b/Engine/Plugins/Experimental/ControlFlows/Source/ControlFlows/Public/ControlFlowConcurrency.h
@@ -4,6 +4,7 @@
 
 #include "Containers/Map.h"
 #include "Delegates/DelegateCombinations.h"
+#include "Delegates/Delegate.h"
 #include "Templates/SharedPointer.h"
 
 class FControlFlow;
