If the Concurrent Control Flow Task was never created, and cancel was called, we should just broadcast that the Control Flow Task was Cancelled

diff --git a/Engine/Plugins/Experimental/ControlFlows/Source/ControlFlows/Private/ControlFlowTask.cpp b/Engine/Plugins/Experimental/ControlFlows/Source/ControlFlows/Private/ControlFlowTask.cpp
--- a/Engine/Plugins/Experimental/ControlFlows/Source/ControlFlows/Private/ControlFlowTask.cpp
+++ b/Engine/Plugins/Experimental/ControlFlows/Source/ControlFlows/Private/ControlFlowTask.cpp
@@ -331,10 +331,14 @@ void FControlFlowTask_ConcurrentFlows::Execute()
 
 void FControlFlowTask_ConcurrentFlows::Cancel()
 {
-	if (ensureAlwaysMsgf(ConcurrentFlows.IsValid(), TEXT("Did you call execute flow on a subflow? Do not do this! You only need to call ExecuteFlow once per FControlFlowStatics::Create!")))
+	if (ConcurrentFlows.IsValid())
 	{
 		ConcurrentFlows->CancelAll();
 	}
+	else
+	{
+		HandleConcurrentFlowsCancelled();
+	}
 }
 
 void FControlFlowTask_ConcurrentFlows::HandleConcurrentFlowsCompleted()
