[ControlFlows]
* Added support for TrackedActivity so we can visualize various control flow's state in the console window
* Changed verbose logging a little bit to skip name of task node when task has a proper name. Couldn't find a case where task node didn't have the name unnamed

diff --git a/Engine/Plugins/Experimental/ControlFlows/Source/ControlFlows/Private/ControlFlow.cpp b/Engine/Plugins/Experimental/ControlFlows/Source/ControlFlows/Private/ControlFlow.cpp
--- a/Engine/Plugins/Experimental/ControlFlows/Source/ControlFlows/Private/ControlFlow.cpp
+++ b/Engine/Plugins/Experimental/ControlFlows/Source/ControlFlows/Private/ControlFlow.cpp
@@ -4,6 +4,7 @@
 #include "ControlFlowManager.h"
 #include "ControlFlowTask.h"
 #include "ControlFlows.h"
+#include "Misc/TrackedActivity.h"
 
 static const int32 MAX_FLOW_LOOPS = 10000;
 
@@ -27,6 +28,10 @@ void FControlFlow::ExecuteNextNodeInQueue()
 	UE_LOG(LogControlFlows, Verbose, TEXT("ControlFlow - Executing %s (ExecuteNextInQueue)"), *GetFlowPath().Append(DebugName));
 
 	CurrentNode = FlowQueue[0];
+	if (Activity)
+		Activity->Update(*CurrentNode->GetNodeName());
+
+
 	FlowQueue.RemoveAt(0);
 	CurrentNode->Execute();
 }
@@ -73,6 +78,9 @@ void FControlFlow::HandleControlFlowNodeCompleted(TSharedRef<const FControlFlowN
 				if (bCancelRequested)
 				{
 					CurrentNode = FlowQueue[0];
+					if (Activity)
+						Activity->Update(*CurrentNode->GetNodeName());
+
 					FlowQueue.RemoveAt(0);
 					CurrentNode->CancelFlow();
 				}
@@ -242,6 +250,19 @@ TOptional<FString> FControlFlow::GetCurrentStepDebugName() const
 	return TOptional<FString>();
 }
 
+TSharedPtr<FTrackedActivity> FControlFlow::GetTrackedActivity() const
+{
+	return Activity;
+}
+
+FControlFlow& FControlFlow::TrackActivities(TSharedPtr<FTrackedActivity> InActivity)
+{
+	if (!InActivity)
+		InActivity = MakeShared<FTrackedActivity>(*DebugName, TEXT(""));
+	Activity = InActivity;
+	return *this;
+}
+
 FSimpleDelegate& FControlFlow::QueueFunction(const FString& FlowNodeDebugName)
 {
 	TSharedRef<FControlFlowNode_SelfCompleting> NewNode = MakeShared<FControlFlowNode_SelfCompleting>(SharedThis(this), FormatOrGetNewNodeDebugName(FlowNodeDebugName));
@@ -262,7 +283,7 @@ FControlFlowPopulator& FControlFlow::QueueControlFlow(const FString& TaskName /*
 {
 	TSharedRef<FControlFlowSimpleSubTask> NewTask = MakeShared<FControlFlowSimpleSubTask>(TaskName, MakeShared<FControlFlow>(TaskName));
 	TSharedRef<FControlFlowNode_Task> NewNode = MakeShared<FControlFlowNode_Task>(SharedThis(this), NewTask, FormatOrGetNewNodeDebugName(FlowNodeDebugName));
-
+	NewTask->GetTaskFlow()->Activity = Activity;
 	NewNode->OnExecute().BindSP(SharedThis(this), &FControlFlow::HandleTaskNodeExecuted);
 	NewNode->OnCancelRequested().BindSP(SharedThis(this), &FControlFlow::HandleTaskNodeCancelled);
 
@@ -276,6 +297,7 @@ FControlFlowBranchDefiner& FControlFlow::QueueControlFlowBranch(const FString& T
 	TSharedRef<FControlFlowTask_Branch> NewTask = MakeShared<FControlFlowTask_Branch>(TaskName);
 	TSharedRef<FControlFlowNode_Task> NewNode = MakeShared<FControlFlowNode_Task>(SharedThis(this), NewTask, FormatOrGetNewNodeDebugName(FlowNodeDebugName));
 
+	NewTask->Activity = Activity;
 	NewNode->OnExecute().BindSP(SharedThis(this), &FControlFlow::HandleTaskNodeExecuted);
 	NewNode->OnCancelRequested().BindSP(SharedThis(this), &FControlFlow::HandleTaskNodeCancelled);
 
@@ -316,7 +338,7 @@ FControlFlowPopulator& FControlFlow::QueueLoop(FControlFlowLoopComplete& LoopCom
 {
 	TSharedRef<FControlFlowTask_Loop> NewTask = MakeShared<FControlFlowTask_Loop>(LoopCompleteDelgate, TaskName, MakeShared<FControlFlow>(TaskName));
 	TSharedRef<FControlFlowNode_Task> NewNode = MakeShared<FControlFlowNode_Task>(SharedThis(this), NewTask, FormatOrGetNewNodeDebugName(FlowNodeDebugName));
-
+	NewTask->GetTaskFlow()->Activity = Activity;
 	NewNode->OnExecute().BindSP(SharedThis(this), &FControlFlow::HandleTaskNodeExecuted);
 	NewNode->OnCancelRequested().BindSP(SharedThis(this), &FControlFlow::HandleTaskNodeCancelled);
 
@@ -333,6 +355,8 @@ void FControlFlow::HandleTaskNodeExecuted(TSharedRef<FControlFlowNode_Task> Task
 	{
 		CurrentlyRunningTask = TaskNode;
 		CurrentNode = CurrentlyRunningTask;
+		if (Activity)
+			Activity->Update(*TaskNode->GetFlowTask()->DebugName);
 
 		TaskNode->GetFlowTask()->OnComplete().BindSP(SharedThis(this), &FControlFlow::HandleOnTaskComplete);
 		TaskNode->GetFlowTask()->OnCancelled().BindSP(SharedThis(this), &FControlFlow::HandleOnTaskCancelled);
diff --git a/Engine/Plugins/Experimental/ControlFlows/Source/ControlFlows/Private/ControlFlowManager.cpp b/Engine/Plugins/Experimental/ControlFlows/Source/ControlFlows/Private/ControlFlowManager.cpp
--- a/Engine/Plugins/Experimental/ControlFlows/Source/ControlFlows/Private/ControlFlowManager.cpp
+++ b/Engine/Plugins/Experimental/ControlFlows/Source/ControlFlows/Private/ControlFlowManager.cpp
@@ -132,6 +132,7 @@ bool FControlFlowStatics::IterateForInvalidFlows(float DeltaTime)
 				TSharedRef<FControlFlow> ExecutingFlow = Executing[Idx]->GetControlFlow();
 				if (!ExecutingFlow->IsRunning() && ensureAlways(ExecutingFlow->NumInQueue() == 0))
 				{
+					Executing[Idx]->ControlFlow->Activity = nullptr;
 					FinishedFlows.Add(Executing[Idx]);
 					Executing.RemoveAtSwap(Idx);
 					--Idx;
@@ -139,6 +140,7 @@ bool FControlFlowStatics::IterateForInvalidFlows(float DeltaTime)
 			}
 			else
 			{
+				Executing[Idx]->ControlFlow->Activity = nullptr;
 				Executing.RemoveAtSwap(Idx);
 				--Idx;
 			}
diff --git a/Engine/Plugins/Experimental/ControlFlows/Source/ControlFlows/Private/ControlFlowNode.cpp b/Engine/Plugins/Experimental/ControlFlows/Source/ControlFlows/Private/ControlFlowNode.cpp
--- a/Engine/Plugins/Experimental/ControlFlows/Source/ControlFlows/Private/ControlFlowNode.cpp
+++ b/Engine/Plugins/Experimental/ControlFlows/Source/ControlFlows/Private/ControlFlowNode.cpp
@@ -42,6 +42,15 @@ void FControlFlowNode::CancelFlow()
 	}
 }
 
+TSharedPtr<FTrackedActivity> FControlFlowNode::GetTrackedActivity() const
+{
+	if (ensureAlways(Parent.IsValid()))
+	{
+		return Parent.Pin()->GetTrackedActivity();
+	}
+	return nullptr;
+}
+
 void FControlFlowNode::LogExecution()
 {
 	if (ensureAlways(Parent.IsValid()))
@@ -142,7 +151,8 @@ void FControlFlowNode_Task::CancelFlow()
 
 FString FControlFlowNode_Task::GetNodeName() const
 {
-	return FString::Format(TEXT("{0}({1})"), { NodeName, GetFlowTask()->GetTaskName() });
+	const FString& TaskName = GetFlowTask()->GetTaskName();
+	return TaskName.IsEmpty() ? NodeName : TaskName;
 }
 
 void FControlFlowNode_Task::CompleteCancelFlow()
diff --git a/Engine/Plugins/Experimental/ControlFlows/Source/ControlFlows/Private/ControlFlowTask.cpp b/Engine/Plugins/Experimental/ControlFlows/Source/ControlFlows/Private/ControlFlowTask.cpp
--- a/Engine/Plugins/Experimental/ControlFlows/Source/ControlFlows/Private/ControlFlowTask.cpp
+++ b/Engine/Plugins/Experimental/ControlFlows/Source/ControlFlows/Private/ControlFlowTask.cpp
@@ -10,6 +10,7 @@
 //////////////////////////
 
 FControlFlowSubTaskBase::FControlFlowSubTaskBase(const FString& TaskName)
+:	DebugName(TaskName)
 {
 
 }
@@ -252,6 +253,7 @@ void FControlFlowTask_Branch::Execute()
 		if (ensureAlwaysMsgf(BranchDefinitions->Contains(SelectedBranchKey), TEXT("You've returned a Branch Key that doesn't exist!")))
 		{
 			SelectedBranchFlow = BranchDefinitions->FindChecked(SelectedBranchKey);
+			SelectedBranchFlow->Activity = Activity;
 
 			SelectedBranchFlow->OnComplete().BindSP(SharedThis(this), &FControlFlowTask_Branch::HandleBranchCompleted);
 			SelectedBranchFlow->OnExecutedWithoutAnyNodes().BindSP(SharedThis(this), &FControlFlowTask_Branch::HandleBranchCompleted);
diff --git a/Engine/Plugins/Experimental/ControlFlows/Source/ControlFlows/Public/ControlFlow.h b/Engine/Plugins/Experimental/ControlFlows/Source/ControlFlows/Public/ControlFlow.h
--- a/Engine/Plugins/Experimental/ControlFlows/Source/ControlFlows/Public/ControlFlow.h
+++ b/Engine/Plugins/Experimental/ControlFlows/Source/ControlFlows/Public/ControlFlow.h
@@ -67,6 +67,8 @@ public:
 
 	TOptional<FString> GetCurrentStepDebugName() const;
 
+	TSharedPtr<FTrackedActivity> GetTrackedActivity() const;
+
 public:
 
 	template<typename...ArgsT>
@@ -112,6 +114,8 @@ public:
 		return *this;
 	}
 
+	FControlFlow& TrackActivities(TSharedPtr<FTrackedActivity> InActivity = nullptr);
+
 public:
 	FSimpleDelegate& QueueFunction(const FString& FlowNodeDebugName = TEXT(""));
 	FControlFlowWaitDelegate& QueueWait(const FString& FlowNodeDebugName = TEXT(""));
@@ -219,6 +223,7 @@ private:
 	friend class FControlFlowTask_Loop;
 	friend class FControlFlowTask_BranchLegacy;
 	friend class FControlFlowTask_Branch;
+	friend class FControlFlowStatics;
 	friend struct FConcurrencySubFlowContainer;
 
 public:
@@ -290,4 +295,6 @@ private:
 	TArray<TSharedRef<FControlFlow>> SubFlowStack_ForDebugging;
 
 	TArray<TSharedRef<FControlFlowNode>> FlowQueue;
+
+	TSharedPtr<FTrackedActivity> Activity;
 };
\ No newline at end of file
diff --git a/Engine/Plugins/Experimental/ControlFlows/Source/ControlFlows/Public/ControlFlowNode.h b/Engine/Plugins/Experimental/ControlFlows/Source/ControlFlows/Public/ControlFlowNode.h
--- a/Engine/Plugins/Experimental/ControlFlows/Source/ControlFlows/Public/ControlFlowNode.h
+++ b/Engine/Plugins/Experimental/ControlFlows/Source/ControlFlows/Public/ControlFlowNode.h
@@ -12,6 +12,7 @@ class FControlFlowSubTaskBase;
 class FControlFlowBranch;
 class FConcurrentControlFlows;
 class FConcurrentControlFlowBehavior;
+class FTrackedActivity;
 
 using FControlFlowNodePtr = TSharedPtr<FControlFlowNode>;
 using FControlFlowNodeRef = TSharedRef<FControlFlowNode>;
@@ -38,6 +39,8 @@ public:
 	CONTROLFLOWS_API void ContinueFlow();
 	CONTROLFLOWS_API virtual void CancelFlow();
 
+	CONTROLFLOWS_API TSharedPtr<FTrackedActivity> GetTrackedActivity() const;
+
 protected:
 	friend class FControlFlow;
 
diff --git a/Engine/Plugins/Experimental/ControlFlows/Source/ControlFlows/Public/ControlFlowTask.h b/Engine/Plugins/Experimental/ControlFlows/Source/ControlFlows/Public/ControlFlowTask.h
--- a/Engine/Plugins/Experimental/ControlFlows/Source/ControlFlows/Public/ControlFlowTask.h
+++ b/Engine/Plugins/Experimental/ControlFlows/Source/ControlFlows/Public/ControlFlowTask.h
@@ -139,6 +139,8 @@ public:
 
 	FControlFlowBranchDefiner& GetDelegate() { return BranchDelegate; }
 
+	TSharedPtr<FTrackedActivity> Activity;
+
 protected:
 	virtual void Execute() override;
 	virtual void Cancel() override;
