Two crash fixes:
1) Calling AsShared() during the destructor a TSharedFromThis
2) Resetting ConcurrentFlows when there are nested ConcurrentFlows

#tests locally repro bug and verified fix

diff --git a/Engine/Plugins/Experimental/ControlFlows/Source/ControlFlows/Private/ControlFlowConcurrency.cpp b/Engine/Plugins/Experimental/ControlFlows/Source/ControlFlows/Private/ControlFlowConcurrency.cpp
--- a/Engine/Plugins/Experimental/ControlFlows/Source/ControlFlows/Private/ControlFlowConcurrency.cpp
+++ b/Engine/Plugins/Experimental/ControlFlows/Source/ControlFlows/Private/ControlFlowConcurrency.cpp
@@ -138,7 +138,7 @@ void FConcurrentControlFlows::HandleConcurrentFlowCompleted(int32 FlowIndex)
 		{
 			if (AreAllSubFlowsCompletedOrCancelled())
 			{
-				OnAllCompleted();
+				OnConcurrencyCompleted.ExecuteIfBound();
 			}
 			else
 			{
diff --git a/Engine/Plugins/Experimental/ControlFlows/Source/ControlFlows/Public/ControlFlow.h b/Engine/Plugins/Experimental/ControlFlows/Source/ControlFlows/Public/ControlFlow.h
--- a/Engine/Plugins/Experimental/ControlFlows/Source/ControlFlows/Public/ControlFlow.h
+++ b/Engine/Plugins/Experimental/ControlFlows/Source/ControlFlows/Public/ControlFlow.h
@@ -142,7 +142,10 @@ private:
 	template<typename BindingObjectClassT, typename...ArgsT>
 	void QueueStep_Internal_DeduceBindingObject(const FString& InDebugName, typename TEnableIf<TIsDerivedFrom<BindingObjectClassT, TSharedFromThis<BindingObjectClassT>>::IsDerived, BindingObjectClassT*>::Type InBindingObject, ArgsT...Params)
 	{
-		QueueStep_Internal_TSharedFromThis(InDebugName, InBindingObject->AsShared(), Params...);
+		if (static_cast<TSharedFromThis<BindingObjectClassT>*>(InBindingObject)->DoesSharedInstanceExist())
+		{
+			QueueStep_Internal_TSharedFromThis(InDebugName, InBindingObject->AsShared(), Params...);
+		}
 	}
 
 	template<typename BindingObjectClassT, typename...ArgsT>
