Fixed ControlFlow TSharedFromThis deduction for derived types

struct Foo : public TSharedFromThis<Foo>
{
    // ...
};

struct Bar : public Foo
{
    // ...

    void Execute()
    {
        FControlFlow& ControlFlow = ...;
        ControlFlow.QueueStep(this, &Bar::DoStep);
    }

    void DoStep()
    {
    }
};

The QueueStep deduction would fail as it expects the type to be directly derived as TSharedFromThis<Bar> rather than indirectly via TSharedFromThis<Foo>

#rnx

diff --git a/Engine/Plugins/Experimental/ControlFlows/Source/ControlFlows/Public/ControlFlow.h b/Engine/Plugins/Experimental/ControlFlows/Source/ControlFlows/Public/ControlFlow.h
--- a/Engine/Plugins/Experimental/ControlFlows/Source/ControlFlows/Public/ControlFlow.h
+++ b/Engine/Plugins/Experimental/ControlFlows/Source/ControlFlows/Public/ControlFlow.h
@@ -140,11 +140,11 @@ private:
 
 private:
 	template<typename BindingObjectClassT, typename...ArgsT>
-	void QueueStep_Internal_DeduceBindingObject(const FString& InDebugName, typename TEnableIf<TIsDerivedFrom<BindingObjectClassT, TSharedFromThis<BindingObjectClassT>>::IsDerived, BindingObjectClassT*>::Type InBindingObject, ArgsT...Params)
+	void QueueStep_Internal_DeduceBindingObject(const FString& InDebugName, typename TEnableIf<IsDerivedFromSharedFromThis<BindingObjectClassT>(), BindingObjectClassT*>::Type InBindingObject, ArgsT...Params)
 	{
-		if (static_cast<TSharedFromThis<BindingObjectClassT>*>(InBindingObject)->DoesSharedInstanceExist())
+		if (InBindingObject->DoesSharedInstanceExist())
 		{
-			QueueStep_Internal_TSharedFromThis(InDebugName, InBindingObject->AsShared(), Params...);
+			QueueStep_Internal_TSharedFromThis(InDebugName, StaticCastSharedRef<BindingObjectClassT>(InBindingObject->AsShared()), Params...);
 		}
 	}
 
diff --git a/Engine/Plugins/Experimental/ControlFlows/Source/ControlFlows/Public/ControlFlowContainer.h b/Engine/Plugins/Experimental/ControlFlows/Source/ControlFlows/Public/ControlFlowContainer.h
--- a/Engine/Plugins/Experimental/ControlFlows/Source/ControlFlows/Public/ControlFlowContainer.h
+++ b/Engine/Plugins/Experimental/ControlFlows/Source/ControlFlows/Public/ControlFlowContainer.h
@@ -59,7 +59,7 @@ public:
 	FControlFlowContainerInternal() = delete;
 	FControlFlowContainerInternal(OwningObjectT* InOwner, TSharedRef<FControlFlow> InFlow, const FString& FlowId)
 		: FControlFlowContainerBase(InFlow, FlowId)
-		, OwningObject(InOwner->AsShared())
+		, OwningObject(StaticCastSharedRef<OwningObjectT>(InOwner->AsShared()))
 	{}
 protected:
 	virtual const void* const GetOwningObject() const override { return OwningObject.IsValid() ? OwningObject.Pin().Get() : nullptr; }
@@ -67,7 +67,7 @@ private:
 	TWeakPtr<const OwningObjectT> OwningObject;
 };
 
-#define FControlFlowContainerInternal_Decl FControlFlowContainerInternal<ExternalObjectT, TIsDerivedFrom<ExternalObjectT, UObject>::IsDerived, TIsDerivedFrom<ExternalObjectT, TSharedFromThis<ExternalObjectT>>::IsDerived>
+#define FControlFlowContainerInternal_Decl FControlFlowContainerInternal<ExternalObjectT, TIsDerivedFrom<ExternalObjectT, UObject>::IsDerived, IsDerivedFromSharedFromThis<ExternalObjectT>()>
 
 template<typename ExternalObjectT>
 class TControlFlowContainer : public FControlFlowContainerInternal_Decl
diff --git a/Engine/Source/Runtime/Core/Public/Templates/SharedPointer.h b/Engine/Source/Runtime/Core/Public/Templates/SharedPointer.h
--- a/Engine/Source/Runtime/Core/Public/Templates/SharedPointer.h
+++ b/Engine/Source/Runtime/Core/Public/Templates/SharedPointer.h
@@ -1720,6 +1720,27 @@ private:
 };
 
 
+namespace UE::Core::Private
+{
+	template <typename T>
+	constexpr bool IsDerivedFromSharedFromThisImpl(const TSharedFromThis<T>*)
+	{
+		return true;
+	}
+
+	constexpr bool IsDerivedFromSharedFromThisImpl(...)
+	{
+		return false;
+	}
+}
+
+template <typename T>
+constexpr bool IsDerivedFromSharedFromThis()
+{
+	return UE::Core::Private::IsDerivedFromSharedFromThisImpl((const T*)nullptr);
+}
+
+
 /**
  * Global equality operator for TSharedRef
  *
