Fix a crash due to global/static destruction order

diff --git a/Engine/Plugins/Experimental/ControlFlows/Source/ControlFlows/Private/ControlFlowManager.cpp b/Engine/Plugins/Experimental/ControlFlows/Source/ControlFlows/Private/ControlFlowManager.cpp
--- a/Engine/Plugins/Experimental/ControlFlows/Source/ControlFlows/Private/ControlFlowManager.cpp
+++ b/Engine/Plugins/Experimental/ControlFlows/Source/ControlFlows/Private/ControlFlowManager.cpp
@@ -4,31 +4,30 @@
 #include "Containers/Ticker.h"
 #include "ControlFlows.h"
 
-static TArray<TSharedRef<FControlFlowContainerBase>> NewlyCreatedFlows;
-static TArray<TSharedRef<FControlFlowContainerBase>> PersistentFlows;
-static TArray<TSharedRef<FControlFlowContainerBase>> ExecutingFlows;
-static TArray<TSharedRef<FControlFlowContainerBase>> FinishedFlows;
-
 static FTSTicker::FDelegateHandle NextFrameCheckForExecution;
 static FTSTicker::FDelegateHandle NextFrameCheckForFlowCleanup;
 
 TArray<TSharedRef<FControlFlowContainerBase>>& FControlFlowStatics::GetNewlyCreatedFlows()
 {
+	static TArray<TSharedRef<FControlFlowContainerBase>> NewlyCreatedFlows;
 	return NewlyCreatedFlows;
 }
 
 TArray<TSharedRef<FControlFlowContainerBase>>& FControlFlowStatics::GetPersistentFlows()
 {
+	static TArray<TSharedRef<FControlFlowContainerBase>> PersistentFlows;
 	return PersistentFlows;
 }
 
 TArray<TSharedRef<FControlFlowContainerBase>>& FControlFlowStatics::GetExecutingFlows()
 {
+	static TArray<TSharedRef<FControlFlowContainerBase>> ExecutingFlows;
 	return ExecutingFlows;
 }
 
 TArray<TSharedRef<FControlFlowContainerBase>>& FControlFlowStatics::GetFinishedFlows()
 {
+	static TArray<TSharedRef<FControlFlowContainerBase>> FinishedFlows;
 	return FinishedFlows;
 }
 
@@ -109,7 +108,7 @@ bool FControlFlowStatics::IterateForInvalidFlows(float DeltaTime)
 				TSharedRef<FControlFlow> PersistentFlow = Persistent[Idx]->GetControlFlow();
 				if (PersistentFlow->IsRunning())
 				{
-					ExecutingFlows.Add(Persistent[Idx]);
+					GetExecutingFlows().Add(Persistent[Idx]);
 					Persistent.RemoveAtSwap(Idx);
 					--Idx;
 				}
@@ -133,7 +132,7 @@ bool FControlFlowStatics::IterateForInvalidFlows(float DeltaTime)
 				if (!ExecutingFlow->IsRunning() && ensureAlways(ExecutingFlow->NumInQueue() == 0))
 				{
 					Executing[Idx]->ControlFlow->Activity = nullptr;
-					FinishedFlows.Add(Executing[Idx]);
+					GetFinishedFlows().Add(Executing[Idx]);
 					Executing.RemoveAtSwap(Idx);
 					--Idx;
 				}
