Rename in preparation to deprecate previous Control Flow Loop implementation

diff --git a/Engine/Plugins/Experimental/ControlFlows/Source/ControlFlows/Private/ControlFlow.cpp b/Engine/Plugins/Experimental/ControlFlows/Source/ControlFlows/Private/ControlFlow.cpp
--- a/Engine/Plugins/Experimental/ControlFlows/Source/ControlFlows/Private/ControlFlow.cpp
+++ b/Engine/Plugins/Experimental/ControlFlows/Source/ControlFlows/Private/ControlFlow.cpp
@@ -336,7 +336,7 @@ TSharedRef<FControlFlowTask_BranchLegacy> FControlFlow::QueueBranch(FControlFlow
 
 FControlFlowPopulator& FControlFlow::QueueLoop(FControlFlowLoopComplete& LoopCompleteDelgate, const FString& TaskName /*= TEXT("")*/, const FString& FlowNodeDebugName /*= TEXT("")*/)
 {
-	TSharedRef<FControlFlowTask_Loop> NewTask = MakeShared<FControlFlowTask_Loop>(LoopCompleteDelgate, TaskName, MakeShared<FControlFlow>(TaskName));
+	TSharedRef<FControlFlowTask_LoopDeprecated> NewTask = MakeShared<FControlFlowTask_LoopDeprecated>(LoopCompleteDelgate, TaskName, MakeShared<FControlFlow>(TaskName));
 	TSharedRef<FControlFlowNode_Task> NewNode = MakeShared<FControlFlowNode_Task>(SharedThis(this), NewTask, FormatOrGetNewNodeDebugName(FlowNodeDebugName));
 	NewTask->GetTaskFlow()->Activity = Activity;
 	NewNode->OnExecute().BindSP(SharedThis(this), &FControlFlow::HandleTaskNodeExecuted);
diff --git a/Engine/Plugins/Experimental/ControlFlows/Source/ControlFlows/Private/ControlFlowTask.cpp b/Engine/Plugins/Experimental/ControlFlows/Source/ControlFlows/Private/ControlFlowTask.cpp
--- a/Engine/Plugins/Experimental/ControlFlows/Source/ControlFlows/Private/ControlFlowTask.cpp
+++ b/Engine/Plugins/Experimental/ControlFlows/Source/ControlFlows/Private/ControlFlowTask.cpp
@@ -176,14 +176,14 @@ void FControlFlowSimpleSubTask::CancelledSubTask()
 //FControlFlowLoop
 //////////////////////////////////
 
-FControlFlowTask_Loop::FControlFlowTask_Loop(FControlFlowLoopComplete& TaskCompleteDelegate, const FString& TaskName, TSharedRef<FControlFlow> FlowOwner)
+FControlFlowTask_LoopDeprecated::FControlFlowTask_LoopDeprecated(FControlFlowLoopComplete& TaskCompleteDelegate, const FString& TaskName, TSharedRef<FControlFlow> FlowOwner)
 	: FControlFlowSimpleSubTask(TaskName, FlowOwner)
 	, TaskCompleteDecider(TaskCompleteDelegate)
 {
 
 }
 
-void FControlFlowTask_Loop::Execute()
+void FControlFlowTask_LoopDeprecated::Execute()
 {
 	if (GetTaskPopulator().IsBound() && TaskCompleteDecider.IsBound() && GetTaskFlow().IsValid())
 	{
@@ -193,9 +193,9 @@ void FControlFlowTask_Loop::Execute()
 		}
 		else
 		{
-			GetTaskFlow()->OnComplete().BindSP(SharedThis(this), &FControlFlowTask_Loop::CompletedLoop);
-			GetTaskFlow()->OnExecutedWithoutAnyNodes().BindSP(SharedThis(this), &FControlFlowTask_Loop::CompletedLoop);
-			GetTaskFlow()->OnCancelled().BindSP(SharedThis(this), &FControlFlowTask_Loop::CancelledLoop);
+			GetTaskFlow()->OnComplete().BindSP(SharedThis(this), &FControlFlowTask_LoopDeprecated::CompletedLoop);
+			GetTaskFlow()->OnExecutedWithoutAnyNodes().BindSP(SharedThis(this), &FControlFlowTask_LoopDeprecated::CompletedLoop);
+			GetTaskFlow()->OnCancelled().BindSP(SharedThis(this), &FControlFlowTask_LoopDeprecated::CancelledLoop);
 
 			GetTaskPopulator().Execute(GetTaskFlow().ToSharedRef());
 
@@ -210,7 +210,7 @@ void FControlFlowTask_Loop::Execute()
 	}
 }
 
-void FControlFlowTask_Loop::Cancel()
+void FControlFlowTask_LoopDeprecated::Cancel()
 {
 	if (GetTaskFlow().IsValid() && GetTaskFlow()->IsRunning())
 	{
@@ -222,7 +222,7 @@ void FControlFlowTask_Loop::Cancel()
 	}
 }
 
-void FControlFlowTask_Loop::CompletedLoop()
+void FControlFlowTask_LoopDeprecated::CompletedLoop()
 {
 	if (TaskCompleteDecider.IsBound() && !TaskCompleteDecider.Execute())
 	{
@@ -234,7 +234,7 @@ void FControlFlowTask_Loop::CompletedLoop()
 	}
 }
 
-void FControlFlowTask_Loop::CancelledLoop()
+void FControlFlowTask_LoopDeprecated::CancelledLoop()
 {
 	OnCancelled().ExecuteIfBound();
 }
diff --git a/Engine/Plugins/Experimental/ControlFlows/Source/ControlFlows/Public/ControlFlow.h b/Engine/Plugins/Experimental/ControlFlows/Source/ControlFlows/Public/ControlFlow.h
--- a/Engine/Plugins/Experimental/ControlFlows/Source/ControlFlows/Public/ControlFlow.h
+++ b/Engine/Plugins/Experimental/ControlFlows/Source/ControlFlows/Public/ControlFlow.h
@@ -223,7 +223,7 @@ private:
 	friend class FControlFlowNode_RequiresCallback;
 	friend class FControlFlowNode_SelfCompleting;
 	friend class FControlFlowSimpleSubTask;
-	friend class FControlFlowTask_Loop;
+	friend class FControlFlowTask_LoopDeprecated;
 	friend class FControlFlowTask_BranchLegacy;
 	friend class FControlFlowTask_Branch;
 	friend class FControlFlowStatics;
diff --git a/Engine/Plugins/Experimental/ControlFlows/Source/ControlFlows/Public/ControlFlowTask.h b/Engine/Plugins/Experimental/ControlFlows/Source/ControlFlows/Public/ControlFlowTask.h
--- a/Engine/Plugins/Experimental/ControlFlows/Source/ControlFlows/Public/ControlFlowTask.h
+++ b/Engine/Plugins/Experimental/ControlFlows/Source/ControlFlows/Public/ControlFlowTask.h
@@ -59,11 +59,11 @@ private:
 
 //////////////////////////////////////////////////////////////////////////////
 
-class FControlFlowTask_Loop : public FControlFlowSimpleSubTask
+class FControlFlowTask_LoopDeprecated : public FControlFlowSimpleSubTask
 {
 public:
 
-	FControlFlowTask_Loop(FControlFlowLoopComplete& TaskCompleteDelegate, const FString& TaskName, TSharedRef<FControlFlow> FlowOwner);
+	FControlFlowTask_LoopDeprecated(FControlFlowLoopComplete& TaskCompleteDelegate, const FString& TaskName, TSharedRef<FControlFlow> FlowOwner);
 
 protected:
 
