Control Flow Changes expansion with public multicast delegates

diff --git a/Engine/Plugins/Experimental/ControlFlows/Source/ControlFlows/Private/ControlFlow.cpp b/Engine/Plugins/Experimental/ControlFlows/Source/ControlFlows/Private/ControlFlow.cpp
--- a/Engine/Plugins/Experimental/ControlFlows/Source/ControlFlows/Private/ControlFlow.cpp
+++ b/Engine/Plugins/Experimental/ControlFlows/Source/ControlFlows/Private/ControlFlow.cpp
@@ -29,8 +29,9 @@ void FControlFlow::ExecuteNextNodeInQueue()
 
 	CurrentNode = FlowQueue[0];
 	if (Activity)
+	{
 		Activity->Update(*CurrentNode->GetNodeName());
-
+	}
 
 	FlowQueue.RemoveAt(0);
 	CurrentNode->Execute();
@@ -45,7 +46,8 @@ void FControlFlow::ExecuteNode(TSharedRef<FControlFlowNode_SelfCompleting> SelfC
 	{
 		FlowQueue.Reset();
 
-		OnCancelled().ExecuteIfBound();
+		OnCancelledDelegate.Broadcast();
+		OnCancelledDelegate_Internal.ExecuteIfBound();		
 
 		FControlFlowStatics::HandleControlFlowFinishedNotification();
 	}
@@ -86,6 +88,7 @@ void FControlFlow::HandleControlFlowNodeCompleted(TSharedRef<const FControlFlowN
 				}
 				else
 				{
+					OnStepCompletedDelegate.Broadcast();
 					ExecuteNextNodeInQueue();
 				}
 			}
@@ -94,11 +97,13 @@ void FControlFlow::HandleControlFlowNodeCompleted(TSharedRef<const FControlFlowN
 		{
 			if (bCancelRequested)
 			{
-				OnCancelled().ExecuteIfBound();
+				OnCancelledDelegate.Broadcast();
+				OnCancelledDelegate_Internal.ExecuteIfBound();				
 			}
 			else
 			{
-				OnComplete().ExecuteIfBound();
+				OnCompleteDelegate.Broadcast();
+				OnCompleteDelegate_Internal.ExecuteIfBound();				
 			}
 
 			FControlFlowStatics::HandleControlFlowFinishedNotification();
@@ -179,8 +184,8 @@ void FControlFlow::ExecuteFlow()
 		}
 		else
 		{
-			OnExecutedWithoutAnyNodes().ExecuteIfBound();
-
+			OnExecutedWithoutAnyNodesDelegate.Broadcast();
+			OnExecutedWithoutAnyNodesDelegate_Internal.ExecuteIfBound();
 			FControlFlowStatics::HandleControlFlowFinishedNotification();
 		}
 
@@ -213,7 +218,8 @@ void FControlFlow::CancelFlow()
 	}
 	else
 	{
-		OnCancelled().ExecuteIfBound();
+		OnCancelledDelegate.Broadcast();
+		OnCancelledDelegate_Internal.ExecuteIfBound();
 		FControlFlowStatics::HandleControlFlowFinishedNotification();
 	}
 }
@@ -356,7 +362,9 @@ void FControlFlow::HandleTaskNodeExecuted(TSharedRef<FControlFlowNode_Task> Task
 		CurrentlyRunningTask = TaskNode;
 		CurrentNode = CurrentlyRunningTask;
 		if (Activity)
+		{
 			Activity->Update(*TaskNode->GetFlowTask()->DebugName);
+		}
 
 		TaskNode->GetFlowTask()->OnComplete().BindSP(SharedThis(this), &FControlFlow::HandleOnTaskComplete);
 		TaskNode->GetFlowTask()->OnCancelled().BindSP(SharedThis(this), &FControlFlow::HandleOnTaskCancelled);
diff --git a/Engine/Plugins/Experimental/ControlFlows/Source/ControlFlows/Private/ControlFlowConcurrency.cpp b/Engine/Plugins/Experimental/ControlFlows/Source/ControlFlows/Private/ControlFlowConcurrency.cpp
--- a/Engine/Plugins/Experimental/ControlFlows/Source/ControlFlows/Private/ControlFlowConcurrency.cpp
+++ b/Engine/Plugins/Experimental/ControlFlows/Source/ControlFlows/Private/ControlFlowConcurrency.cpp
@@ -48,17 +48,17 @@ void FConcurrencySubFlowContainer::Cancel()
 
 FSimpleDelegate& FConcurrencySubFlowContainer::OnComplete()
 {
-	return SubFlow->OnComplete();
+	return SubFlow->OnCompleteDelegate_Internal;
 }
 
 FSimpleDelegate& FConcurrencySubFlowContainer::OnExecutedWithoutAnyNodes()
 {
-	return SubFlow->OnExecutedWithoutAnyNodes();
+	return SubFlow->OnExecutedWithoutAnyNodesDelegate_Internal;
 }
 
 FSimpleDelegate& FConcurrencySubFlowContainer::OnCancelled()
 {
-	return SubFlow->OnCancelled();
+	return SubFlow->OnCancelledDelegate_Internal;
 }
 
 const FString& FConcurrencySubFlowContainer::GetDebugName()
diff --git a/Engine/Plugins/Experimental/ControlFlows/Source/ControlFlows/Private/ControlFlowTask.cpp b/Engine/Plugins/Experimental/ControlFlows/Source/ControlFlows/Private/ControlFlowTask.cpp
--- a/Engine/Plugins/Experimental/ControlFlows/Source/ControlFlows/Private/ControlFlowTask.cpp
+++ b/Engine/Plugins/Experimental/ControlFlows/Source/ControlFlows/Private/ControlFlowTask.cpp
@@ -93,9 +93,9 @@ void FControlFlowTask_BranchLegacy::Execute()
 
 		TSharedRef<FControlFlow> FlowToExecute = GetOrAddBranch(SelectedBranch);
 
-		FlowToExecute->OnComplete().BindSP(SharedThis(this), &FControlFlowTask_BranchLegacy::HandleBranchCompleted);
-		FlowToExecute->OnExecutedWithoutAnyNodes().BindSP(SharedThis(this), &FControlFlowTask_BranchLegacy::HandleBranchCompleted);
-		FlowToExecute->OnCancelled().BindSP(SharedThis(this), &FControlFlowTask_BranchLegacy::HandleBranchCancelled);
+		FlowToExecute->OnCompleteDelegate_Internal.BindSP(SharedThis(this), &FControlFlowTask_BranchLegacy::HandleBranchCompleted);
+		FlowToExecute->OnExecutedWithoutAnyNodesDelegate_Internal.BindSP(SharedThis(this), &FControlFlowTask_BranchLegacy::HandleBranchCompleted);
+		FlowToExecute->OnCancelledDelegate_Internal.BindSP(SharedThis(this), &FControlFlowTask_BranchLegacy::HandleBranchCancelled);
 
 		FlowToExecute->ExecuteFlow();
 	}
@@ -132,9 +132,9 @@ void FControlFlowSimpleSubTask::Execute()
 {
 	if (TaskPopulator.IsBound() && GetTaskFlow().IsValid())
 	{
-		GetTaskFlow()->OnComplete().BindSP(SharedThis(this), &FControlFlowSimpleSubTask::CompletedSubTask);
-		GetTaskFlow()->OnExecutedWithoutAnyNodes().BindSP(SharedThis(this), &FControlFlowSimpleSubTask::CompletedSubTask);
-		GetTaskFlow()->OnCancelled().BindSP(SharedThis(this), &FControlFlowSimpleSubTask::CancelledSubTask);
+		GetTaskFlow()->OnCompleteDelegate_Internal.BindSP(SharedThis(this), &FControlFlowSimpleSubTask::CompletedSubTask);
+		GetTaskFlow()->OnExecutedWithoutAnyNodesDelegate_Internal.BindSP(SharedThis(this), &FControlFlowSimpleSubTask::CompletedSubTask);
+		GetTaskFlow()->OnCancelledDelegate_Internal.BindSP(SharedThis(this), &FControlFlowSimpleSubTask::CancelledSubTask);
 
 		TaskPopulator.Execute(GetTaskFlow().ToSharedRef());
 
@@ -179,9 +179,7 @@ void FControlFlowSimpleSubTask::CancelledSubTask()
 FControlFlowTask_LoopDeprecated::FControlFlowTask_LoopDeprecated(FControlFlowLoopComplete& TaskCompleteDelegate, const FString& TaskName, TSharedRef<FControlFlow> FlowOwner)
 	: FControlFlowSimpleSubTask(TaskName, FlowOwner)
 	, TaskCompleteDecider(TaskCompleteDelegate)
-{
-
-}
+{}
 
 void FControlFlowTask_LoopDeprecated::Execute()
 {
@@ -193,9 +191,9 @@ void FControlFlowTask_LoopDeprecated::Execute()
 		}
 		else
 		{
-			GetTaskFlow()->OnComplete().BindSP(SharedThis(this), &FControlFlowTask_LoopDeprecated::CompletedLoop);
-			GetTaskFlow()->OnExecutedWithoutAnyNodes().BindSP(SharedThis(this), &FControlFlowTask_LoopDeprecated::CompletedLoop);
-			GetTaskFlow()->OnCancelled().BindSP(SharedThis(this), &FControlFlowTask_LoopDeprecated::CancelledLoop);
+			GetTaskFlow()->OnCompleteDelegate_Internal.BindSP(SharedThis(this), &FControlFlowTask_LoopDeprecated::CompletedLoop);
+			GetTaskFlow()->OnExecutedWithoutAnyNodesDelegate_Internal.BindSP(SharedThis(this), &FControlFlowTask_LoopDeprecated::CompletedLoop);
+			GetTaskFlow()->OnCancelledDelegate_Internal.BindSP(SharedThis(this), &FControlFlowTask_LoopDeprecated::CancelledLoop);
 
 			GetTaskPopulator().Execute(GetTaskFlow().ToSharedRef());
 
@@ -255,9 +253,9 @@ void FControlFlowTask_Branch::Execute()
 			SelectedBranchFlow = BranchDefinitions->FindChecked(SelectedBranchKey);
 			SelectedBranchFlow->Activity = Activity;
 
-			SelectedBranchFlow->OnComplete().BindSP(SharedThis(this), &FControlFlowTask_Branch::HandleBranchCompleted);
-			SelectedBranchFlow->OnExecutedWithoutAnyNodes().BindSP(SharedThis(this), &FControlFlowTask_Branch::HandleBranchCompleted);
-			SelectedBranchFlow->OnCancelled().BindSP(SharedThis(this), &FControlFlowTask_Branch::HandleBranchCancelled);
+			SelectedBranchFlow->OnCompleteDelegate_Internal.BindSP(SharedThis(this), &FControlFlowTask_Branch::HandleBranchCompleted);
+			SelectedBranchFlow->OnExecutedWithoutAnyNodesDelegate_Internal.BindSP(SharedThis(this), &FControlFlowTask_Branch::HandleBranchCompleted);
+			SelectedBranchFlow->OnCancelledDelegate_Internal.BindSP(SharedThis(this), &FControlFlowTask_Branch::HandleBranchCancelled);
 
 			ensureAlwaysMsgf(!BranchDefinitions->IsAnyBranchRunning(), TEXT("Did you call ExecuteFlow() on a Branch? Do not do this! You only need to call ExecuteFlow once per FControlFlowStatics::Create!"));
 
diff --git a/Engine/Plugins/Experimental/ControlFlows/Source/ControlFlows/Public/ControlFlow.h b/Engine/Plugins/Experimental/ControlFlows/Source/ControlFlows/Public/ControlFlow.h
--- a/Engine/Plugins/Experimental/ControlFlows/Source/ControlFlows/Public/ControlFlow.h
+++ b/Engine/Plugins/Experimental/ControlFlows/Source/ControlFlows/Public/ControlFlow.h
@@ -52,6 +52,7 @@ class CONTROLFLOWS_API FControlFlow : public TSharedFromThis<FControlFlow>
 {
 public:
 	FControlFlow(const FString& FlowDebugName = TEXT(""));
+	virtual ~FControlFlow() = default;
 
 public:
 	/** This needs to be called, otherwise nothing will happen!! Call after you finish adding functions to the queue. Calling with an empty queue is safe. */
@@ -239,25 +240,29 @@ public:
 	/** Adds a Loop to your flow. The flow will use FControlFlowLoopComplete - if this returns false, the flow will execute FControlTaskQueuePopulator until true is returned */
 	FControlFlowPopulator& QueueLoop(FControlFlowLoopComplete& LoopCompleteDelgate, const FString& TaskName = TEXT(""), const FString& FlowNodeDebugName = TEXT(""));
 
+	FSimpleMulticastDelegate& OnNodeComplete() const { return OnStepCompletedDelegate; }
+	FSimpleMulticastDelegate& OnComplete() const { return OnCompleteDelegate; }
+	FSimpleMulticastDelegate& OnCancel() const { return OnCancelledDelegate; }
+	FSimpleMulticastDelegate& OnExecuteWithoutAnyNodes() const { return OnExecutedWithoutAnyNodesDelegate; }	
+	
 private:
-
 	void HandleControlFlowNodeCompleted(TSharedRef<const FControlFlowNode> NodeCompleted);
 
-	void ExecuteNextNodeInQueue();
-
-	FSimpleDelegate& OnComplete() const { return OnCompleteDelegate; }
-	FSimpleDelegate& OnExecutedWithoutAnyNodes() const { return OnExecutedWithoutAnyNodesDelegate; }
-	FSimpleDelegate& OnCancelled() const { return OnCancelledDelegate; }
-
-	mutable FSimpleDelegate OnCompleteDelegate;
-	mutable FSimpleDelegate OnExecutedWithoutAnyNodesDelegate;
-	mutable FSimpleDelegate OnCancelledDelegate;
+	mutable FSimpleMulticastDelegate OnStepCompletedDelegate;
+	mutable FSimpleMulticastDelegate OnExecutedWithoutAnyNodesDelegate;
+	mutable FSimpleMulticastDelegate OnCompleteDelegate;
+	mutable FSimpleMulticastDelegate OnCancelledDelegate;
+	
+	mutable FSimpleDelegate OnCompleteDelegate_Internal;
+	mutable FSimpleDelegate OnExecutedWithoutAnyNodesDelegate_Internal;
+	mutable FSimpleDelegate OnCancelledDelegate_Internal;
 
 private:
+	void ExecuteNextNodeInQueue();
+
 	void ExecuteNode(TSharedRef<FControlFlowNode_SelfCompleting> SelfCompletingNode);
 
 private:
-
 	void HandleTaskNodeExecuted(TSharedRef<FControlFlowNode_Task> TaskNode);
 	void HandleTaskNodeCancelled(TSharedRef<FControlFlowNode_Task> TaskNode);
 
