Expanded ControlFlows with Loop Syntax:

FlowInstance
	.Loop([this](TSharedRef<FConditionalLoop> LoopHandle)
	{
		LoopHandle->RunLoopFirst()                       // See FConditionalLoop inside ControlFlowConditionalLoop.h for additional options
			.QueueStep(this, &ThisClass::Foo)
			.QueueStep(this, &ThisClass::Foo);

		return FMath::RandBool() ? ELoopConditions::RunLoop : ELoopConditions::LoopFinished;
	});

diff --git a/Engine/Plugins/Experimental/ControlFlows/Source/ControlFlows/Private/ControlFlow.cpp b/Engine/Plugins/Experimental/ControlFlows/Source/ControlFlows/Private/ControlFlow.cpp
--- a/Engine/Plugins/Experimental/ControlFlows/Source/ControlFlows/Private/ControlFlow.cpp
+++ b/Engine/Plugins/Experimental/ControlFlows/Source/ControlFlows/Private/ControlFlow.cpp
@@ -3,6 +3,7 @@
 #include "ControlFlow.h"
 #include "ControlFlowManager.h"
 #include "ControlFlowTask.h"
+#include "ControlFlowTask_Loop.h"
 #include "ControlFlows.h"
 #include "Misc/TrackedActivity.h"
 
@@ -40,8 +41,15 @@ void FControlFlow::ExecuteNextNodeInQueue()
 void FControlFlow::ExecuteNode(TSharedRef<FControlFlowNode_SelfCompleting> SelfCompletingNode)
 {
 	// Calling Synchronous Function in Flow
-	SelfCompletingNode->Process.ExecuteIfBound();
-
+	if (SelfCompletingNode->Process.IsBound())
+	{
+		SelfCompletingNode->Process.Execute();
+	}
+	else
+	{
+		SelfCompletingNode->bWasBoundOnExecution = false;
+	}
+	
 	if (SelfCompletingNode->HasCancelBeenRequested())
 	{
 		FlowQueue.Reset();
@@ -65,6 +73,11 @@ void FControlFlow::HandleControlFlowNodeCompleted(TSharedRef<const FControlFlowN
 
 	if (ensure(CurrentNode.IsValid() && &NodeCompleted.Get() == CurrentNode.Get()))
 	{
+		if (!NodeCompleted->bWasBoundOnExecution)
+		{
+			OnNodeWasNotBoundedOnExecution_Internal.ExecuteIfBound();
+		}
+
 		const bool bCancelRequested = !bInterpretCancelledNodeAsComplete && NodeCompleted->HasCancelBeenRequested();
 		CurrentNode.Reset();
 		CurrentlyRunningTask.Reset();
@@ -175,21 +188,28 @@ void FControlFlow::ExecuteFlow()
 
 	if (ensureAlwaysMsgf(!CurrentNode.IsValid(), TEXT("Flow is already running! Or perhaps there are multiple instances of owning class? All flows should have a unique ID.")))
 	{
-		FControlFlowStatics::HandleControlFlowStartedNotification(AsShared());
-
-		SubFlowStack_ForDebugging.Add(SharedThis(this));
-		if (FlowQueue.Num() > 0)
+		if (!ensure(SubFlowStack_ForDebugging.Num() < MAX_FLOW_LOOPS))
 		{
-			ExecuteNextNodeInQueue();
+			UE_LOG(LogControlFlows, Error, TEXT("ControlFlow - Hit maximum Flow loops. Is there an infinite recursion somewhere?"));
 		}
 		else
 		{
-			OnExecutedWithoutAnyNodesDelegate.Broadcast();
-			OnExecutedWithoutAnyNodesDelegate_Internal.ExecuteIfBound();
-			FControlFlowStatics::HandleControlFlowFinishedNotification();
-		}
+			FControlFlowStatics::HandleControlFlowStartedNotification(AsShared());
+
+			SubFlowStack_ForDebugging.Add(SharedThis(this));
+			if (FlowQueue.Num() > 0)
+			{
+				ExecuteNextNodeInQueue();
+			}
+			else
+			{
+				OnExecutedWithoutAnyNodesDelegate.Broadcast();
+				OnExecutedWithoutAnyNodesDelegate_Internal.ExecuteIfBound();
+				FControlFlowStatics::HandleControlFlowFinishedNotification();
+			}
 
-		SubFlowStack_ForDebugging.Pop();
+			SubFlowStack_ForDebugging.Pop();
+		}
 	}
 
 	UE_LOG(LogControlFlows, Verbose, TEXT("ControlFlow - Executing %s (ExecuteFlow Completed)"), *DebugName);
@@ -325,6 +345,19 @@ FConcurrentFlowsDefiner& FControlFlow::QueueConcurrentFlows(const FString& TaskN
 	return NewTask->GetDelegate();
 }
 
+FControlFlowConditionalLoopDefiner& FControlFlow::QueueConditionalLoop(const FString& TaskName /*= TEXT("")*/, const FString& FlowNodeDebugName /*= TEXT("")*/)
+{
+	TSharedRef<FControlFlowTask_ConditionalLoop> NewTask = MakeShared<FControlFlowTask_ConditionalLoop>(TaskName);
+	TSharedRef<FControlFlowNode_Task> NewNode = MakeShared<FControlFlowNode_Task>(SharedThis(this), NewTask, FormatOrGetNewNodeDebugName(FlowNodeDebugName));
+
+	NewNode->OnExecute().BindSP(SharedThis(this), &FControlFlow::HandleTaskNodeExecuted);
+	NewNode->OnCancelRequested().BindSP(SharedThis(this), &FControlFlow::HandleTaskNodeCancelled);
+
+	FlowQueue.Add(NewNode);
+
+	return NewTask->GetDelegate();
+}
+
 TSharedRef<FControlFlowTask_BranchLegacy> FControlFlow::QueueBranch(FControlFlowBranchDecider_Legacy& BranchDecider, const FString& TaskName /*= TEXT("")*/, const FString& FlowNodeDebugName /*= TEXT("")*/)
 {
 	ensureAlwaysMsgf(false, TEXT("Deprecated. Use 'QueueControlFlowBranch'"));
diff --git a/Engine/Plugins/Experimental/ControlFlows/Source/ControlFlows/Private/ControlFlowConditionalLoop.cpp b/Engine/Plugins/Experimental/ControlFlows/Source/ControlFlows/Private/ControlFlowConditionalLoop.cpp
new file mode 100644
--- /dev/null
+++ b/Engine/Plugins/Experimental/ControlFlows/Source/ControlFlows/Private/ControlFlowConditionalLoop.cpp
@@ -0,0 +1,29 @@
+// Copyright Epic Games, Inc. All Rights Reserved.
+
+#include "ControlFlowConditionalLoop.h"
+#include "ControlFlow.h"
+
+FConditionalLoop::FConditionalLoop()
+	: FlowLoop(MakeShared<FControlFlow>())
+{}
+
+FControlFlow& FConditionalLoop::SetCheckConditionFirst(bool bInValue)
+{
+	CheckConditionalFirst.Emplace(bInValue);
+	return FlowLoop.Get();
+}
+
+FControlFlow& FConditionalLoop::ExecuteAtLeastOnce()
+{
+	return SetCheckConditionFirst(false);
+}
+
+FControlFlow& FConditionalLoop::CheckConditionFirst()
+{
+	return SetCheckConditionFirst(true);
+}
+
+FControlFlow& FConditionalLoop::RunLoopFirst()
+{
+	return SetCheckConditionFirst(false);
+}
diff --git a/Engine/Plugins/Experimental/ControlFlows/Source/ControlFlows/Private/ControlFlowNode.cpp b/Engine/Plugins/Experimental/ControlFlows/Source/ControlFlows/Private/ControlFlowNode.cpp
--- a/Engine/Plugins/Experimental/ControlFlows/Source/ControlFlows/Private/ControlFlowNode.cpp
+++ b/Engine/Plugins/Experimental/ControlFlows/Source/ControlFlows/Private/ControlFlowNode.cpp
@@ -115,6 +115,7 @@ void FControlFlowNode_RequiresCallback::Execute()
 	}
 	else
 	{
+		bWasBoundOnExecution = false;
 		ContinueFlow();
 	}
 }
@@ -138,6 +139,7 @@ void FControlFlowNode_Task::Execute()
 	}
 	else
 	{
+		bWasBoundOnExecution = false;
 		ContinueFlow();
 	}
 }
diff --git a/Engine/Plugins/Experimental/ControlFlows/Source/ControlFlows/Private/ControlFlowTask_Loop.cpp b/Engine/Plugins/Experimental/ControlFlows/Source/ControlFlows/Private/ControlFlowTask_Loop.cpp
new file mode 100644
--- /dev/null
+++ b/Engine/Plugins/Experimental/ControlFlows/Source/ControlFlows/Private/ControlFlowTask_Loop.cpp
@@ -0,0 +1,98 @@
+// Copyright Epic Games, Inc. All Rights Reserved.
+
+#include "ControlFlowTask_Loop.h"
+
+#include "ControlFlows.h"
+#include "ControlFlowConditionalLoop.h"
+
+void FControlFlowTask_ConditionalLoop::Execute()
+{
+	if (!bStopLoopingOnUnboundedDelegate && ensureAlways(ConditionalLoopDelegate.IsBound() && !ConditionalLoop.IsValid()))
+	{
+		ConditionalLoop = MakeShared<FConditionalLoop>();
+
+		ConditionalLoop->FlowLoop->OnCompleteDelegate_Internal.BindSP(SharedThis(this), &FControlFlowTask_ConditionalLoop::HandleLoopCompleted);
+		ConditionalLoop->FlowLoop->OnExecutedWithoutAnyNodesDelegate_Internal.BindSP(SharedThis(this), &FControlFlowTask_ConditionalLoop::HandleLoopCompleted);
+		ConditionalLoop->FlowLoop->OnCancelledDelegate_Internal.BindSP(SharedThis(this), &FControlFlowTask_ConditionalLoop::HandleLoopCancelled);
+		ConditionalLoop->FlowLoop->OnNodeWasNotBoundedOnExecution_Internal.BindSP(SharedThis(this), &FControlFlowTask_ConditionalLoop::HandleOnNodeWasNotBoundedOnExecution);
+
+		const EConditionalLoopResult InitialConditionalValue = ConditionalLoopDelegate.Execute(ConditionalLoop.ToSharedRef());
+
+		if (ConditionalLoop->FlowLoop->IsRunning())
+		{
+			ensureAlwaysMsgf(false, TEXT("Did you call ExecuteFlow() on a Loop? Do not do this! You only need to call ExecuteFlow once per FControlFlowStatics::Create!"));
+		}
+		else
+		{
+			if (ensureAlwaysMsgf(ConditionalLoop->CheckConditionalFirst.IsSet(), TEXT("Unused Coonditional Loop Handle. Will skip this step.")))
+			{
+				bHandleUsed = true;
+				if (ConditionalLoop->CheckConditionalFirst.GetValue())
+				{
+					Helper_ConditionalCheck(InitialConditionalValue);
+				}
+				else
+				{
+					ConditionalLoop->FlowLoop->ExecuteFlow();
+				}
+			}
+			else
+			{
+				OnComplete().ExecuteIfBound();
+			}
+		}
+	}
+	else
+	{
+		OnComplete().ExecuteIfBound();
+	}
+}
+
+void FControlFlowTask_ConditionalLoop::Cancel()
+{
+	if (ensureAlways(ConditionalLoop.IsValid()) && ConditionalLoop->FlowLoop->IsRunning())
+	{
+		ConditionalLoop->FlowLoop->CancelFlow();
+	}
+	else
+	{
+		HandleLoopCancelled();
+	}
+}
+void FControlFlowTask_ConditionalLoop::Helper_ConditionalCheck(const EConditionalLoopResult InCondition) const
+{
+	if (InCondition == EConditionalLoopResult::RunLoop)
+	{
+		ConditionalLoop->FlowLoop->ExecuteFlow();
+	}
+	else if (InCondition == EConditionalLoopResult::LoopFinished)
+	{
+		OnComplete().ExecuteIfBound();
+	}
+	else
+	{
+		ensureAlwaysMsgf(false, TEXT("Unhandled Enum value"));
+		OnComplete().ExecuteIfBound();
+	}
+}
+
+void FControlFlowTask_ConditionalLoop::HandleLoopCompleted()
+{
+	ensureMsgf(bHandleUsed, TEXT("Did you call ExecuteFlow() on a Loop? Do not do this! This can cause a stack overflow!"));
+
+	if (!bStopLoopingOnUnboundedDelegate && ensureAlways(ConditionalLoopDelegate.IsBound() && ConditionalLoop.IsValid()))
+	{
+		Helper_ConditionalCheck(ConditionalLoopDelegate.Execute(ConditionalLoop.ToSharedRef()));
+	}
+}
+
+void FControlFlowTask_ConditionalLoop::HandleLoopCancelled()
+{
+	OnCancelled().ExecuteIfBound();
+}
+
+void FControlFlowTask_ConditionalLoop::HandleOnNodeWasNotBoundedOnExecution()
+{
+	UE_LOG(LogControlFlows, Warning, TEXT("[Conditional Loop: %s] %s"), ANSI_TO_TCHAR(__FUNCTION__), *GetTaskName());
+	bStopLoopingOnUnboundedDelegate = true;
+}
diff --git a/Engine/Plugins/Experimental/ControlFlows/Source/ControlFlows/Public/ControlFlow.h b/Engine/Plugins/Experimental/ControlFlows/Source/ControlFlows/Public/ControlFlow.h
--- a/Engine/Plugins/Experimental/ControlFlows/Source/ControlFlows/Public/ControlFlow.h
+++ b/Engine/Plugins/Experimental/ControlFlows/Source/ControlFlows/Public/ControlFlow.h
@@ -21,15 +21,25 @@ class FControlFlowTask_BranchLegacy;
  *
  *  'QueueWait': Queues a 'void (FControlFlowNodeRef FlowHandle, ...)' function.
  *               The flow will stop until 'FlowHandle->ContinueFlow()' is called.
- *	             BE RESPONSIBLE and make sure all code paths call to continue, otherwise the flow will hang.
+ *	             BE RESPONSIBLE and make sure all code paths call to continue xOR cancel, otherwise the flow will hang.
  *
  *  'QueueControlFlow': Queues a 'void (TSharedRef<FControlFlow> Subflow, ...)' function.
  * 
- *  'QueueControlFlowBranch': Queues a 'int32(TSharedRef<FControlFlowBranch> Branch, ...)' function. TODO: Allow any return-type for branch selection, and not restricted to only int32.
+ *  'QueueControlFlowBranch': Queues a 'int32(TSharedRef<FControlFlowBranch> Branch, ...)' function.
+ *							  TODO: Allow any return-type for branch selection, and not restricted to only int32.
+ *							  Lambda-Syntax: '.BranchFlow([this](TSharedRef<FControlFlowBranch> Branch) { })'
+ * 
+ *  'QueueConcurrentFlows': Queues a 'void(TSharedRef<FConcurrentControlFlows> ForkedFlow, ...)' function
+ *							Lambda-Syntax: '.ForkFlow([this](TSharedRef<FConcurrentControlFlows> Fork) { })'
+ * 
+ *  'QueueConditionalLoop': Queues a 'EConditionalLoopResult(TSharedRef<FConditionalLoop> Loop, ...)' function
+ *							The public functions of 'FConditionalLoop' forces the caller to define the loop as as a 'while (CONDITION) {}' or a 'do {} while(CONDITION);'
+ *							The return is a FControlFlow& and queue your flow as normal
+ *							Lambda-Syntax: '.Loop([this](TSharedRef<FConditionalLoop> Loop) { })'
  * 
  *  'QueueStep': Usable in #UObject's or classes that derive from #TSharedFromThis<UserClass>. The Control Flow will automatically deduce if this is a
- *				 '#QueueFunction', '#QueueWait', '#QueueControlFlow', '#QueueControlFlowBranch' based on the function signature.
- *				 Returns a ref to the ControlFlow - enables chained step queue-ing.
+ *				 '#QueueFunction', '#QueueWait', '#QueueControlFlow', '#QueueControlFlowBranch', '#QueueConcurrentFlows', '#EConditionalLoopResult' based on the function signature.
+ *				 Returns a ref to the ControlFlow - enables chained step queueing.
  *
  *  Using the auto-deduction of 'QueueStep', you can change the queue from a synchronous function (QueueFunction) to an asynchronous one (QueueWait) or vice-versa
  *  by adding/removing the 'FControlFlowNodeRef FlowHandle' as your first parameter. And you can change it to (QueueControlFlow) if need be as well!
@@ -46,6 +56,70 @@ class FControlFlowTask_BranchLegacy;
  *		.QueueStep(this, &UserClass:MyFunction3, ...);
  * 
  *  This allow ease of going from Synchronous Functionality to Asynchronously Functionality to Subflows as you build out your Flow.
+ * 
+ *  ******************
+ *  Full Example Class
+ * 
+ *	struct FMyFlowClass : public TSharedFromThis<FMyFlowClass>
+ *	{
+ *		typedef FMyFlowClass ThisClass;
+ * 
+ * 		FMyVariableContainer() : MyPurpose(MakeShared<FControlFlow>()) {}
+ *
+ * 		void RunMyPurpose()
+ * 		{
+ * 			MyPurpose
+ *				.QueueStep(this, &MyFlowClass::Construct)
+ *				.Loop([this](TSharedRef<FConditionalLoop> Outerloop)
+ *				{
+ *					Outerloop->RunLoopFirst()
+ *						.SetCancelledNodeAsComplete(true)
+ *						.QueueStep(this, &ThisClass::Foo)
+ *						.QueueStep(this, &ThisClass::Foo)
+ *						.Loop([this](TSharedRef<FConditionalLoop> InnerLoop)
+ *						{
+ *							InnerLoop->CheckConditionFirst()
+ *								.BranchFlow([this], TSharedRef<FConditionalLoop> Branch)
+ *								{
+ *									Branch->AddOrGetBranch(0)
+ *										.QueueStep(this, &ThisClass::Foo)
+ *										.QueueStep(this, &ThisClass::Foo);
+ * 
+ *									Branch->AddOrGetBranch(1)
+ *										.QueueStep(this, &ThisClass::Foo)
+ *										.ForkFlow([this], TSharedRef<FConcurrentControlFlows> ConcurrentFlows)
+ *										{
+ *											ConcurrentFlows->AddOrGetFlow(0)
+ *												.QueueStep(this, &ThisClass::Foo)
+ *												.QueueStep(this, &ThisClass::Foo);
+ * 
+ *											ConcurrentFlows->AddOrGetFlow(1)
+ *												.QueueStep(this, &ThisClass::Foo);
+ *										});
+ * 
+ *									return FMath::RandBool() ? 0 : 1;
+ *								})
+ *								.QueueStep(this, &ThisClass::Foo);
+ *
+ *							return FMath::RandBool() ? EConditionalLoopResult::RunLoop : EConditionalLoopResult::LoopFinished;
+ *						})
+ *						.QueueStep(this, &ThisClass::Foo);
+ *
+ *					return EConditionalLoopResult::RunLoop; // OuterLoop will never end in this example
+ *				})
+ *				.QueueStep(this, &ThisClass::Destruct)
+ *				.ExecuteFlow();
+ *		}
+ *
+ *	private:
+ *		void Foo();
+ *		void Construct(); // Equivalent to Init
+ *		void Destruct(); // Equivalent to Uninit
+ *
+ *		TSharedRef<FControlFlow> MyPurpose;
+ *	};
+ * 
+ * The implementation details of a flow can be fully disjointed from the Flow logic itself!
  */
 
 class CONTROLFLOWS_API FControlFlow : public TSharedFromThis<FControlFlow>
@@ -101,6 +175,13 @@ public:
 	}
 
 public:
+	template<typename FunctionT, typename...ArgsT>
+	FControlFlow& Loop(FunctionT InLoopLambda, ArgsT...Params)
+	{
+		QueueConditionalLoop(FormatOrGetNewNodeDebugName()).BindLambda(InLoopLambda, Params...);
+		return *this;
+	}
+
 	template<typename FunctionT, typename...ArgsT>
 	FControlFlow& BranchFlow(FunctionT InBranchLambda, ArgsT...Params)
 	{
@@ -123,6 +204,7 @@ public:
 	FControlFlowPopulator& QueueControlFlow(const FString& TaskName = TEXT(""), const FString& FlowNodeDebugName = TEXT(""));
 	FControlFlowBranchDefiner& QueueControlFlowBranch(const FString& TaskName = TEXT(""), const FString& FlowNodeDebugName = TEXT(""));
 	FConcurrentFlowsDefiner& QueueConcurrentFlows(const FString& TaskName = TEXT(""), const FString& FlowNodeDebugName = TEXT(""));
+	FControlFlowConditionalLoopDefiner& QueueConditionalLoop(const FString& TaskName = TEXT(""), const FString& FlowNodeDebugName = TEXT(""));
 
 private:
 	template<typename BindingObjectT, typename...PayloadParamsT>
@@ -157,6 +239,12 @@ private:
 
 private:
 	
+	template<typename BindingObjectClassT, typename...PayloadParamsT>
+	void QueueStep_Internal_TSharedFromThis(const FString& InDebugName, TSharedRef<BindingObjectClassT> InBindingObject, typename TMemFunPtrType<false, BindingObjectClassT, EConditionalLoopResult(TSharedRef<FConditionalLoop>, PayloadParamsT...)>::Type InFunction, PayloadParamsT...Params)
+	{
+		QueueConditionalLoop(InDebugName).BindSP(InBindingObject, InFunction, Params...);
+	}
+
 	template<typename BindingObjectClassT, typename...PayloadParamsT>
 	void QueueStep_Internal_TSharedFromThis(const FString& InDebugName, TSharedRef<BindingObjectClassT> InBindingObject, typename TMemFunPtrType<false, BindingObjectClassT, void(TSharedRef<FConcurrentControlFlows>, PayloadParamsT...)>::Type InFunction, PayloadParamsT...Params)
 	{
@@ -189,6 +277,12 @@ private:
 
 private:
 
+	template<typename BindingObjectClassT, typename...PayloadParamsT>
+	void QueueStep_Internal_UObject(const FString& InDebugName, BindingObjectClassT* InBindingObject, typename TMemFunPtrType<false, BindingObjectClassT, EConditionalLoopResult(TSharedRef<FConditionalLoop>, PayloadParamsT...)>::Type InFunction, PayloadParamsT...Params)
+	{
+		QueueConditionalLoop(InDebugName).BindUObject(InBindingObject, InFunction, Params...);
+	}
+
 	template<typename BindingObjectClassT, typename...PayloadParamsT>
 	void QueueStep_Internal_UObject(const FString& InDebugName, BindingObjectClassT* InBindingObject, typename TMemFunPtrType<false, BindingObjectClassT, void(TSharedRef<FConcurrentControlFlows>, PayloadParamsT...)>::Type InFunction, PayloadParamsT...Params)
 	{
@@ -227,6 +321,7 @@ private:
 	friend class FControlFlowTask_LoopDeprecated;
 	friend class FControlFlowTask_BranchLegacy;
 	friend class FControlFlowTask_Branch;
+	friend class FControlFlowTask_ConditionalLoop;
 	friend class FControlFlowStatics;
 	friend struct FConcurrencySubFlowContainer;
 
@@ -256,6 +351,7 @@ private:
 	mutable FSimpleDelegate OnCompleteDelegate_Internal;
 	mutable FSimpleDelegate OnExecutedWithoutAnyNodesDelegate_Internal;
 	mutable FSimpleDelegate OnCancelledDelegate_Internal;
+	mutable FSimpleDelegate OnNodeWasNotBoundedOnExecution_Internal;
 
 private:
 	void ExecuteNextNodeInQueue();
diff --git a/Engine/Plugins/Experimental/ControlFlows/Source/ControlFlows/Public/ControlFlowConditionalLoop.h b/Engine/Plugins/Experimental/ControlFlows/Source/ControlFlows/Public/ControlFlowConditionalLoop.h
new file mode 100644
--- /dev/null
+++ b/Engine/Plugins/Experimental/ControlFlows/Source/ControlFlows/Public/ControlFlowConditionalLoop.h
@@ -0,0 +1,26 @@
+// Copyright Epic Games, Inc. All Rights Reserved.
+
+#pragma once
+
+#include "Templates/SharedPointer.h"
+#include "Misc/Optional.h"
+
+class FControlFlow;
+
+// In contrast to, for example, iterator loops
+class CONTROLFLOWS_API FConditionalLoop : public TSharedFromThis<FConditionalLoop>
+{
+public:
+	FConditionalLoop();
+	
+	FControlFlow& CheckConditionFirst();
+	FControlFlow& RunLoopFirst();
+	FControlFlow& ExecuteAtLeastOnce();
+	FControlFlow& SetCheckConditionFirst(bool bInValue);
+
+private:
+	friend class FControlFlowTask_ConditionalLoop;
+
+	TOptional<bool> CheckConditionalFirst;
+	TSharedRef<FControlFlow> FlowLoop;
+};
\ No newline at end of file
diff --git a/Engine/Plugins/Experimental/ControlFlows/Source/ControlFlows/Public/ControlFlowNode.h b/Engine/Plugins/Experimental/ControlFlows/Source/ControlFlows/Public/ControlFlowNode.h
--- a/Engine/Plugins/Experimental/ControlFlows/Source/ControlFlows/Public/ControlFlowNode.h
+++ b/Engine/Plugins/Experimental/ControlFlows/Source/ControlFlows/Public/ControlFlowNode.h
@@ -9,6 +9,7 @@ class FControlFlowNode;
 class FControlFlowNode_Task;
 class FControlFlowSubTaskBase;
 class FControlFlowBranch;
+class FConditionalLoop;
 class FConcurrentControlFlows;
 class FConcurrentControlFlowBehavior;
 class FTrackedActivity;
@@ -25,6 +26,13 @@ DECLARE_DELEGATE_RetVal(int32, FControlFlowBranchDecider_Legacy)
 // Would be nice to not force people to use int32 for the branch key, but finding the syntax to get it to work properly proved to be more difficult than it was worth. End result is callers have to cast to an int32.
 DECLARE_DELEGATE_RetVal_OneParam(int32, FControlFlowBranchDefiner, TSharedRef<FControlFlowBranch>);
 
+enum class EConditionalLoopResult
+{
+	RunLoop,
+	LoopFinished
+};
+DECLARE_DELEGATE_RetVal_OneParam(EConditionalLoopResult, FControlFlowConditionalLoopDefiner, TSharedRef<FConditionalLoop>);
+
 DECLARE_DELEGATE_OneParam(FConcurrentFlowsDefiner, TSharedRef<FConcurrentControlFlows>);
 
 class FControlFlowNode : public TSharedFromThis<FControlFlowNode>
@@ -53,6 +61,7 @@ protected:
 	TWeakPtr<FControlFlow> Parent = nullptr;
 	FString NodeName;
 	bool bCancelled = false;
+	bool bWasBoundOnExecution = true;
 };
 
 class FControlFlowNode_RequiresCallback : public FControlFlowNode
diff --git a/Engine/Plugins/Experimental/ControlFlows/Source/ControlFlows/Public/ControlFlowTask_Loop.h b/Engine/Plugins/Experimental/ControlFlows/Source/ControlFlows/Public/ControlFlowTask_Loop.h
new file mode 100644
--- /dev/null
+++ b/Engine/Plugins/Experimental/ControlFlows/Source/ControlFlows/Public/ControlFlowTask_Loop.h
@@ -0,0 +1,41 @@
+// Copyright Epic Games, Inc. All Rights Reserved.
+
+#pragma once
+
+#include "ControlFlowTask.h"
+
+class FConditionalLoop;
+
+// In contrast to, for example, iterator loops
+class FControlFlowTask_ConditionalLoop : public FControlFlowSubTaskBase
+{
+public:
+	FControlFlowTask_ConditionalLoop(const FString& TaskName)
+		: FControlFlowSubTaskBase(TaskName)
+	{}
+
+	FControlFlowConditionalLoopDefiner& GetDelegate() { return ConditionalLoopDelegate; }
+
+protected:
+
+	virtual void Execute() override;
+	virtual void Cancel() override;
+
+private:
+	void Helper_ConditionalCheck(const EConditionalLoopResult InCondition) const;
+
+	void HandleLoopCompleted();
+	void HandleLoopCancelled();
+	void HandleOnNodeWasNotBoundedOnExecution();
+
+	FControlFlowConditionalLoopDefiner ConditionalLoopDelegate;
+	TSharedPtr<FConditionalLoop> ConditionalLoop;
+
+	bool bHandleUsed = false;
+
+	/*
+	* While using the '.Loop' syntax, the object doing the bounding can suddenly be destroyed and cause a stack overflow with an infinite loop
+	* 'bStopLoopingOnUnboundedDelegate' is used to prevent this
+	*/
+	bool bStopLoopingOnUnboundedDelegate = false;
+};
\ No newline at end of file
