Fusing delegates for api simplification

diff --git a/Engine/Plugins/Experimental/ControlFlows/Source/ControlFlows/Private/ControlFlow.cpp b/Engine/Plugins/Experimental/ControlFlows/Source/ControlFlows/Private/ControlFlow.cpp
--- a/Engine/Plugins/Experimental/ControlFlows/Source/ControlFlows/Private/ControlFlow.cpp
+++ b/Engine/Plugins/Experimental/ControlFlows/Source/ControlFlows/Private/ControlFlow.cpp
@@ -54,7 +54,7 @@ void FControlFlow::ExecuteNode(TSharedRef<FControlFlowNode_SelfCompleting> SelfC
 	{
 		FlowQueue.Reset();
 
-		OnCancelledDelegate.Broadcast();
+		OnFlowCancelledDelegate.Broadcast();
 		OnCancelledDelegate_Internal.ExecuteIfBound();		
 
 		FControlFlowStatics::HandleControlFlowFinishedNotification();
@@ -110,12 +110,12 @@ void FControlFlow::HandleControlFlowNodeCompleted(TSharedRef<const FControlFlowN
 		{
 			if (bCancelRequested)
 			{
-				OnCancelledDelegate.Broadcast();
+				OnFlowCancelledDelegate.Broadcast();
 				OnCancelledDelegate_Internal.ExecuteIfBound();				
 			}
 			else
 			{
-				OnCompleteDelegate.Broadcast();
+				OnFlowCompleteDelegate.Broadcast();
 				OnCompleteDelegate_Internal.ExecuteIfBound();				
 			}
 
@@ -203,7 +203,8 @@ void FControlFlow::ExecuteFlow()
 			}
 			else
 			{
-				OnExecutedWithoutAnyNodesDelegate.Broadcast();
+				// No nodes is treated as completed for non internal cases
+				OnFlowCompleteDelegate.Broadcast();
 				OnExecutedWithoutAnyNodesDelegate_Internal.ExecuteIfBound();
 				FControlFlowStatics::HandleControlFlowFinishedNotification();
 			}
@@ -238,7 +239,7 @@ void FControlFlow::CancelFlow()
 	}
 	else
 	{
-		OnCancelledDelegate.Broadcast();
+		OnFlowCancelledDelegate.Broadcast();
 		OnCancelledDelegate_Internal.ExecuteIfBound();
 		FControlFlowStatics::HandleControlFlowFinishedNotification();
 	}
diff --git a/Engine/Plugins/Experimental/ControlFlows/Source/ControlFlows/Public/ControlFlow.h b/Engine/Plugins/Experimental/ControlFlows/Source/ControlFlows/Public/ControlFlow.h
--- a/Engine/Plugins/Experimental/ControlFlows/Source/ControlFlows/Public/ControlFlow.h
+++ b/Engine/Plugins/Experimental/ControlFlows/Source/ControlFlows/Public/ControlFlow.h
@@ -126,7 +126,6 @@ class CONTROLFLOWS_API FControlFlow : public TSharedFromThis<FControlFlow>
 {
 public:
 	FControlFlow(const FString& FlowDebugName = TEXT(""));
-	virtual ~FControlFlow() = default;
 
 public:
 	/** This needs to be called, otherwise nothing will happen!! Call after you finish adding functions to the queue. Calling with an empty queue is safe. */
@@ -336,17 +335,15 @@ public:
 	FControlFlowPopulator& QueueLoop(FControlFlowLoopComplete& LoopCompleteDelgate, const FString& TaskName = TEXT(""), const FString& FlowNodeDebugName = TEXT(""));
 
 	FSimpleMulticastDelegate& OnNodeComplete() const { return OnStepCompletedDelegate; }
-	FSimpleMulticastDelegate& OnComplete() const { return OnCompleteDelegate; }
-	FSimpleMulticastDelegate& OnCancel() const { return OnCancelledDelegate; }
-	FSimpleMulticastDelegate& OnExecuteWithoutAnyNodes() const { return OnExecutedWithoutAnyNodesDelegate; }	
+	FSimpleMulticastDelegate& OnFlowComplete() const { return OnFlowCompleteDelegate; }
+	FSimpleMulticastDelegate& OnFlowCancel() const { return OnFlowCancelledDelegate; }
 	
 private:
 	void HandleControlFlowNodeCompleted(TSharedRef<const FControlFlowNode> NodeCompleted);
 
 	mutable FSimpleMulticastDelegate OnStepCompletedDelegate;
-	mutable FSimpleMulticastDelegate OnExecutedWithoutAnyNodesDelegate;
-	mutable FSimpleMulticastDelegate OnCompleteDelegate;
-	mutable FSimpleMulticastDelegate OnCancelledDelegate;
+	mutable FSimpleMulticastDelegate OnFlowCompleteDelegate;
+	mutable FSimpleMulticastDelegate OnFlowCancelledDelegate;
 	
 	mutable FSimpleDelegate OnCompleteDelegate_Internal;
 	mutable FSimpleDelegate OnExecutedWithoutAnyNodesDelegate_Internal;
