[ControlFlows] Add an externally accessible OnCancelled delegate in 23.40

#rnx

diff --git a/Engine/Plugins/Experimental/ControlFlows/Source/ControlFlows/Private/ControlFlow.cpp b/Engine/Plugins/Experimental/ControlFlows/Source/ControlFlows/Private/ControlFlow.cpp
--- a/Engine/Plugins/Experimental/ControlFlows/Source/ControlFlows/Private/ControlFlow.cpp
+++ b/Engine/Plugins/Experimental/ControlFlows/Source/ControlFlows/Private/ControlFlow.cpp
@@ -54,9 +54,7 @@ void FControlFlow::ExecuteNode(TSharedRef<FControlFlowNode_SelfCompleting> SelfC
 	{
 		FlowQueue.Reset();
 
-		OnFlowCancelledDelegate.Broadcast();
-		OnCancelledDelegate_Internal.ExecuteIfBound();		
-
+		BroadcastCancellation();
 		FControlFlowStatics::HandleControlFlowFinishedNotification();
 	}
 	else
@@ -65,6 +63,12 @@ void FControlFlow::ExecuteNode(TSharedRef<FControlFlowNode_SelfCompleting> SelfC
 	}
 }
 
+void FControlFlow::BroadcastCancellation()
+{
+	OnFlowCancelledDelegate.Broadcast();
+	OnCancelledDelegate_Internal.ExecuteIfBound();
+}
+
 void FControlFlow::HandleControlFlowNodeCompleted(TSharedRef<const FControlFlowNode> NodeCompleted)
 {
 	UE_LOG(LogControlFlows, Verbose, TEXT("ControlFlow - Executing %s%s(x%d).%s (FlowControlNodeCompleted)"), *GetFlowPath(), *DebugName, GetRepeatedFlowCount(), *NodeCompleted->GetNodeName());
@@ -110,8 +114,7 @@ void FControlFlow::HandleControlFlowNodeCompleted(TSharedRef<const FControlFlowN
 		{
 			if (bCancelRequested)
 			{
-				OnFlowCancelledDelegate.Broadcast();
-				OnCancelledDelegate_Internal.ExecuteIfBound();				
+				BroadcastCancellation();
 			}
 			else
 			{
@@ -231,6 +234,16 @@ void FControlFlow::Reset()
 	UnnamedBranchCounter = 0;
 }
 
+FDelegateHandle FControlFlow::RegisterOnControlFlowCancelled(FSimpleMulticastDelegate::FDelegate Delegate)
+{
+	return OnFlowCancelledDelegate.Add(Delegate);
+}
+
+void FControlFlow::RemoveOnControlFlowCancelled(FDelegateHandle Handle)
+{
+	OnFlowCancelledDelegate.Remove(Handle);
+}
+
 void FControlFlow::CancelFlow()
 {
 	if (CurrentNode.IsValid())
@@ -239,8 +252,7 @@ void FControlFlow::CancelFlow()
 	}
 	else
 	{
-		OnFlowCancelledDelegate.Broadcast();
-		OnCancelledDelegate_Internal.ExecuteIfBound();
+		BroadcastCancellation();
 		FControlFlowStatics::HandleControlFlowFinishedNotification();
 	}
 }
diff --git a/Engine/Plugins/Experimental/ControlFlows/Source/ControlFlows/Public/ControlFlow.h b/Engine/Plugins/Experimental/ControlFlows/Source/ControlFlows/Public/ControlFlow.h
--- a/Engine/Plugins/Experimental/ControlFlows/Source/ControlFlows/Public/ControlFlow.h
+++ b/Engine/Plugins/Experimental/ControlFlows/Source/ControlFlows/Public/ControlFlow.h
@@ -134,6 +134,9 @@ public:
 	bool IsRunning() const { return CurrentNode.IsValid(); }
 	size_t NumInQueue() const { return FlowQueue.Num(); }
 
+	FDelegateHandle RegisterOnControlFlowCancelled(FSimpleMulticastDelegate::FDelegate Delegate);
+	void RemoveOnControlFlowCancelled(FDelegateHandle Handle);
+
 	/** Will cancel ALL flows, both child ControlFlows and ControlFlows who owns this Flow. You've been warned. */
 	void CancelFlow();
 
@@ -354,7 +357,6 @@ private:
 	void ExecuteNextNodeInQueue();
 
 	void ExecuteNode(TSharedRef<FControlFlowNode_SelfCompleting> SelfCompletingNode);
-
 private:
 	void HandleTaskNodeExecuted(TSharedRef<FControlFlowNode_Task> TaskNode);
 	void HandleTaskNodeCancelled(TSharedRef<FControlFlowNode_Task> TaskNode);
@@ -362,6 +364,7 @@ private:
 	void HandleOnTaskComplete();
 	void HandleOnTaskCancelled();
 
+	void BroadcastCancellation();
 private:
 
 	void LogNodeExecution(const FControlFlowNode& NodeExecuted);
