Allow querying a flow node to see if it's been canceled

#rnx

diff --git a/Engine/Plugins/Experimental/ControlFlows/Source/ControlFlows/Public/ControlFlowNode.h b/Engine/Plugins/Experimental/ControlFlows/Source/ControlFlows/Public/ControlFlowNode.h
--- a/Engine/Plugins/Experimental/ControlFlows/Source/ControlFlows/Public/ControlFlowNode.h
+++ b/Engine/Plugins/Experimental/ControlFlows/Source/ControlFlows/Public/ControlFlowNode.h
@@ -48,10 +48,10 @@ public:
 
 	CONTROLFLOWS_API TSharedPtr<FTrackedActivity> GetTrackedActivity() const;
 
+	[[nodiscard]] bool HasCancelBeenRequested() const { return bCancelled; }
+
 protected:
 	friend class FControlFlow;
-
-	bool HasCancelBeenRequested() const { return bCancelled; }
 	
 	virtual void Execute() {}
 
