Reorganized headers

diff --git a/Engine/Plugins/Experimental/ControlFlows/Source/ControlFlows/Public/ControlFlow.h b/Engine/Plugins/Experimental/ControlFlows/Source/ControlFlows/Public/ControlFlow.h
--- a/Engine/Plugins/Experimental/ControlFlows/Source/ControlFlows/Public/ControlFlow.h
+++ b/Engine/Plugins/Experimental/ControlFlows/Source/ControlFlows/Public/ControlFlow.h
@@ -122,29 +122,33 @@ class FControlFlowTask_BranchLegacy;
  * The implementation details of a flow can be fully disjointed from the Flow logic itself!
  */
 
-class CONTROLFLOWS_API FControlFlow : public TSharedFromThis<FControlFlow>
+class FControlFlow : public TSharedFromThis<FControlFlow>
 {
 public:
-	FControlFlow(const FString& FlowDebugName = TEXT(""));
+	CONTROLFLOWS_API FControlFlow(const FString& FlowDebugName = TEXT(""));
 
 public:
 	/** This needs to be called, otherwise nothing will happen!! Call after you finish adding functions to the queue. Calling with an empty queue is safe. */
-	void ExecuteFlow();
-	void Reset();
+	CONTROLFLOWS_API void ExecuteFlow();
+	CONTROLFLOWS_API void Reset();
 	bool IsRunning() const { return CurrentNode.IsValid(); }
 	size_t NumInQueue() const { return FlowQueue.Num(); }
 
-	FDelegateHandle RegisterOnControlFlowCancelled(FSimpleMulticastDelegate::FDelegate Delegate);
-	void RemoveOnControlFlowCancelled(FDelegateHandle Handle);
+	FSimpleMulticastDelegate& OnNodeComplete() const { return OnStepCompletedDelegate; }
+	FSimpleMulticastDelegate& OnFlowComplete() const { return OnFlowCompleteDelegate; }
+	FSimpleMulticastDelegate& OnFlowCancel() const { return OnFlowCancelledDelegate; }
+
+	CONTROLFLOWS_API FDelegateHandle RegisterOnControlFlowCancelled(FSimpleMulticastDelegate::FDelegate Delegate);
+	CONTROLFLOWS_API void RemoveOnControlFlowCancelled(FDelegateHandle Handle);
 
 	/** Will cancel ALL flows, both child ControlFlows and ControlFlows who owns this Flow. You've been warned. */
-	void CancelFlow();
+	CONTROLFLOWS_API void CancelFlow();
 
-	FControlFlow& SetCancelledNodeAsComplete(bool bCancelledNodeIsComplete);
+	CONTROLFLOWS_API FControlFlow& SetCancelledNodeAsComplete(bool bCancelledNodeIsComplete);
 
-	TOptional<FString> GetCurrentStepDebugName() const;
+	CONTROLFLOWS_API TOptional<FString> GetCurrentStepDebugName() const;
 
-	TSharedPtr<FTrackedActivity> GetTrackedActivity() const;
+	CONTROLFLOWS_API TSharedPtr<FTrackedActivity> GetTrackedActivity() const;
 
 public:
 
@@ -198,15 +202,15 @@ public:
 		return *this;
 	}
 
-	FControlFlow& TrackActivities(TSharedPtr<FTrackedActivity> InActivity = nullptr);
+	CONTROLFLOWS_API FControlFlow& TrackActivities(TSharedPtr<FTrackedActivity> InActivity = nullptr);
 
 public:
-	FSimpleDelegate& QueueFunction(const FString& FlowNodeDebugName = TEXT(""));
-	FControlFlowWaitDelegate& QueueWait(const FString& FlowNodeDebugName = TEXT(""));
-	FControlFlowPopulator& QueueControlFlow(const FString& TaskName = TEXT(""), const FString& FlowNodeDebugName = TEXT(""));
-	FControlFlowBranchDefiner& QueueControlFlowBranch(const FString& TaskName = TEXT(""), const FString& FlowNodeDebugName = TEXT(""));
-	FConcurrentFlowsDefiner& QueueConcurrentFlows(const FString& TaskName = TEXT(""), const FString& FlowNodeDebugName = TEXT(""));
-	FControlFlowConditionalLoopDefiner& QueueConditionalLoop(const FString& TaskName = TEXT(""), const FString& FlowNodeDebugName = TEXT(""));
+	CONTROLFLOWS_API FSimpleDelegate& QueueFunction(const FString& FlowNodeDebugName = TEXT(""));
+	CONTROLFLOWS_API FControlFlowWaitDelegate& QueueWait(const FString& FlowNodeDebugName = TEXT(""));
+	CONTROLFLOWS_API FControlFlowPopulator& QueueControlFlow(const FString& TaskName = TEXT(""), const FString& FlowNodeDebugName = TEXT(""));
+	CONTROLFLOWS_API FControlFlowBranchDefiner& QueueControlFlowBranch(const FString& TaskName = TEXT(""), const FString& FlowNodeDebugName = TEXT(""));
+	CONTROLFLOWS_API FConcurrentFlowsDefiner& QueueConcurrentFlows(const FString& TaskName = TEXT(""), const FString& FlowNodeDebugName = TEXT(""));
+	CONTROLFLOWS_API FControlFlowConditionalLoopDefiner& QueueConditionalLoop(const FString& TaskName = TEXT(""), const FString& FlowNodeDebugName = TEXT(""));
 
 private:
 	template<typename BindingObjectT, typename...PayloadParamsT>
@@ -328,18 +332,9 @@ private:
 	friend struct FConcurrencySubFlowContainer;
 
 public:
-	/** These work, but they are a bit clunky to use. The heart of the issue is that it requires the caller to define two functions. We want only the caller to use one function.
-	  * Not making templated versions of them until a better API is figured out. */	
-	TSharedRef<FControlFlowTask_BranchLegacy> QueueBranch(FControlFlowBranchDecider_Legacy& BranchDecider, const FString& TaskName = TEXT(""), const FString& FlowNodeDebugName = TEXT(""));
-
-	//TODO: Implement #define QueueLoop_Signature TMemFunPtrType<false, BindingObjectClassType, bool(TSharedRef<FControlFlow> SubFlow, VarTypes...)> and delete
-
-	/** Adds a Loop to your flow. The flow will use FControlFlowLoopComplete - if this returns false, the flow will execute FControlTaskQueuePopulator until true is returned */
-	FControlFlowPopulator& QueueLoop(FControlFlowLoopComplete& LoopCompleteDelgate, const FString& TaskName = TEXT(""), const FString& FlowNodeDebugName = TEXT(""));
-
-	FSimpleMulticastDelegate& OnNodeComplete() const { return OnStepCompletedDelegate; }
-	FSimpleMulticastDelegate& OnFlowComplete() const { return OnFlowCompleteDelegate; }
-	FSimpleMulticastDelegate& OnFlowCancel() const { return OnFlowCancelledDelegate; }
+	// Both of these implementations are deprecated
+	CONTROLFLOWS_API TSharedRef<FControlFlowTask_BranchLegacy> QueueBranch(FControlFlowBranchDecider_Legacy& BranchDecider, const FString& TaskName = TEXT(""), const FString& FlowNodeDebugName = TEXT(""));
+	CONTROLFLOWS_API FControlFlowPopulator& QueueLoop(FControlFlowLoopComplete& LoopCompleteDelgate, const FString& TaskName = TEXT(""), const FString& FlowNodeDebugName = TEXT(""));
 	
 private:
 	void HandleControlFlowNodeCompleted(TSharedRef<const FControlFlowNode> NodeCompleted);
@@ -377,7 +372,7 @@ public:
 	const FString& GetDebugName() const { return DebugName; }
 
 private:
-	FString FormatOrGetNewNodeDebugName(const FString& FlowNodeDebugName = TEXT(""));
+	CONTROLFLOWS_API FString FormatOrGetNewNodeDebugName(const FString& FlowNodeDebugName = TEXT(""));
 
 private:
 
diff --git a/Engine/Plugins/Experimental/ControlFlows/Source/ControlFlows/Public/ControlFlowConditionalLoop.h b/Engine/Plugins/Experimental/ControlFlows/Source/ControlFlows/Public/ControlFlowConditionalLoop.h
--- a/Engine/Plugins/Experimental/ControlFlows/Source/ControlFlows/Public/ControlFlowConditionalLoop.h
+++ b/Engine/Plugins/Experimental/ControlFlows/Source/ControlFlows/Public/ControlFlowConditionalLoop.h
@@ -8,15 +8,15 @@
 class FControlFlow;
 
 // In contrast to, for example, iterator loops
-class CONTROLFLOWS_API FConditionalLoop : public TSharedFromThis<FConditionalLoop>
+class FConditionalLoop : public TSharedFromThis<FConditionalLoop>
 {
 public:
-	FConditionalLoop();
+	CONTROLFLOWS_API FConditionalLoop();
 	
-	FControlFlow& CheckConditionFirst();
-	FControlFlow& RunLoopFirst();
-	FControlFlow& ExecuteAtLeastOnce();
-	FControlFlow& SetCheckConditionFirst(bool bInValue);
+	CONTROLFLOWS_API FControlFlow& CheckConditionFirst();
+	CONTROLFLOWS_API FControlFlow& RunLoopFirst();
+	CONTROLFLOWS_API FControlFlow& ExecuteAtLeastOnce();
+	CONTROLFLOWS_API FControlFlow& SetCheckConditionFirst(bool bInValue);
 
 private:
 	friend class FControlFlowTask_ConditionalLoop;
diff --git a/Engine/Plugins/Experimental/ControlFlows/Source/ControlFlows/Public/ControlFlowManager.h b/Engine/Plugins/Experimental/ControlFlows/Source/ControlFlows/Public/ControlFlowManager.h
--- a/Engine/Plugins/Experimental/ControlFlows/Source/ControlFlows/Public/ControlFlowManager.h
+++ b/Engine/Plugins/Experimental/ControlFlows/Source/ControlFlows/Public/ControlFlowManager.h
@@ -8,7 +8,7 @@
 #include "ControlFlowContainer.h"
 // IWYU pragma: end_keep
 
-class CONTROLFLOWS_API FControlFlowStatics
+class FControlFlowStatics
 {
 public:
 	template<typename OwningObjectT>
@@ -73,16 +73,16 @@ public:
 
 private:
 	// These flows will be checked the next frame to make sure they are executed. If they are not, it will execute the flow and remove from this array.
-	static TArray<TSharedRef<FControlFlowContainerBase>>& GetNewlyCreatedFlows();
+	CONTROLFLOWS_API static TArray<TSharedRef<FControlFlowContainerBase>>& GetNewlyCreatedFlows();
 
 	// These flows will not be checked the next frame, and will be moved to executing when we find ourselves executing.
-	static TArray<TSharedRef<FControlFlowContainerBase>>& GetPersistentFlows();
+	CONTROLFLOWS_API static TArray<TSharedRef<FControlFlowContainerBase>>& GetPersistentFlows();
 
 	// Flows that are actively running.
-	static TArray<TSharedRef<FControlFlowContainerBase>>& GetExecutingFlows();
+	CONTROLFLOWS_API static TArray<TSharedRef<FControlFlowContainerBase>>& GetExecutingFlows();
 
 	// Final Array that flows are moved to before being deleted.
-	static TArray<TSharedRef<FControlFlowContainerBase>>& GetFinishedFlows();
+	CONTROLFLOWS_API static TArray<TSharedRef<FControlFlowContainerBase>>& GetFinishedFlows();
 
 private:
 	template<typename OwningObjectT>
@@ -140,8 +140,8 @@ private:
 	static void HandleControlFlowFinishedNotification() { CheckForInvalidFlows(); }
 
 private:
-	static void CheckNewlyCreatedFlows();
-	static void CheckForInvalidFlows();
+	CONTROLFLOWS_API static void CheckNewlyCreatedFlows();
+	CONTROLFLOWS_API static void CheckForInvalidFlows();
 
 private:
 	static bool IterateThroughNewlyCreatedFlows(float DeltaTime);
