Deleted Register OnFlowCancel functions since there is a public MulticastDelegate to bind to

diff --git a/Engine/Plugins/Experimental/ControlFlows/Source/ControlFlows/Private/ControlFlow.cpp b/Engine/Plugins/Experimental/ControlFlows/Source/ControlFlows/Private/ControlFlow.cpp
--- a/Engine/Plugins/Experimental/ControlFlows/Source/ControlFlows/Private/ControlFlow.cpp
+++ b/Engine/Plugins/Experimental/ControlFlows/Source/ControlFlows/Private/ControlFlow.cpp
@@ -234,16 +234,6 @@ void FControlFlow::Reset()
 	UnnamedBranchCounter = 0;
 }
 
-FDelegateHandle FControlFlow::RegisterOnControlFlowCancelled(FSimpleMulticastDelegate::FDelegate Delegate)
-{
-	return OnFlowCancelledDelegate.Add(Delegate);
-}
-
-void FControlFlow::RemoveOnControlFlowCancelled(FDelegateHandle Handle)
-{
-	OnFlowCancelledDelegate.Remove(Handle);
-}
-
 void FControlFlow::CancelFlow()
 {
 	if (CurrentNode.IsValid())
diff --git a/Engine/Plugins/Experimental/ControlFlows/Source/ControlFlows/Public/ControlFlow.h b/Engine/Plugins/Experimental/ControlFlows/Source/ControlFlows/Public/ControlFlow.h
--- a/Engine/Plugins/Experimental/ControlFlows/Source/ControlFlows/Public/ControlFlow.h
+++ b/Engine/Plugins/Experimental/ControlFlows/Source/ControlFlows/Public/ControlFlow.h
@@ -138,9 +138,6 @@ public:
 	FSimpleMulticastDelegate& OnFlowComplete() const { return OnFlowCompleteDelegate; }
 	FSimpleMulticastDelegate& OnFlowCancel() const { return OnFlowCancelledDelegate; }
 
-	CONTROLFLOWS_API FDelegateHandle RegisterOnControlFlowCancelled(FSimpleMulticastDelegate::FDelegate Delegate);
-	CONTROLFLOWS_API void RemoveOnControlFlowCancelled(FDelegateHandle Handle);
-
 	/** Will cancel ALL flows, both child ControlFlows and ControlFlows who owns this Flow. You've been warned. */
 	CONTROLFLOWS_API void CancelFlow();
 
