Corrected Comments

diff --git a/Engine/Plugins/Experimental/ControlFlows/Source/ControlFlows/Public/ControlFlow.h b/Engine/Plugins/Experimental/ControlFlows/Source/ControlFlows/Public/ControlFlow.h
--- a/Engine/Plugins/Experimental/ControlFlows/Source/ControlFlows/Public/ControlFlow.h
+++ b/Engine/Plugins/Experimental/ControlFlows/Source/ControlFlows/Public/ControlFlow.h
@@ -64,12 +64,12 @@ class FControlFlowTask_BranchLegacy;
  *	{
  *		typedef FMyFlowClass ThisClass;
  * 
- * 		FMyVariableContainer() : MyPurpose(MakeShared<FControlFlow>()) {}
+ * 		FMyFlowClass() : MyPurpose(MakeShared<FControlFlow>()) {}
  *
  * 		void RunMyPurpose()
  * 		{
  * 			MyPurpose
- *				.QueueStep(this, &MyFlowClass::Construct)
+ *				.QueueStep(this, &ThisClass::Construct)
  *				.Loop([this](TSharedRef<FConditionalLoop> Outerloop)
  *				{
  *					Outerloop->RunLoopFirst()
