Updated ensure with messaging

diff --git a/Engine/Plugins/Experimental/ControlFlows/Source/ControlFlows/Private/ControlFlow.cpp b/Engine/Plugins/Experimental/ControlFlows/Source/ControlFlows/Private/ControlFlow.cpp
--- a/Engine/Plugins/Experimental/ControlFlows/Source/ControlFlows/Private/ControlFlow.cpp
+++ b/Engine/Plugins/Experimental/ControlFlows/Source/ControlFlows/Private/ControlFlow.cpp
@@ -75,7 +75,7 @@ void FControlFlow::HandleControlFlowNodeCompleted(TSharedRef<const FControlFlowN
 
 	SubFlowStack_ForDebugging.Add(SharedThis(this));
 
-	if (ensure(CurrentNode.IsValid() && &NodeCompleted.Get() == CurrentNode.Get()))
+	if (ensureMsgf(CurrentNode.IsValid() && &NodeCompleted.Get() == CurrentNode.Get(), TEXT("This will likely become a static_assert/compile-error")))
 	{
 		if (!NodeCompleted->bWasBoundOnExecution)
 		{
