Split ```ensure(CurrentNode.IsValid() && &NodeCompleted.Get() == CurrentNode.Get())``` --> ```ensure(CurrentNode.IsValid()) && ensure(&NodeCompleted.Get() == CurrentNode.Get())```.  These are logically equivalent, but more information is given by having separate ensures if one of them fires (and they are still guaranteed not to fire both).

#UE5

diff --git a/Engine/Plugins/Experimental/ControlFlows/Source/ControlFlows/Private/ControlFlow.cpp b/Engine/Plugins/Experimental/ControlFlows/Source/ControlFlows/Private/ControlFlow.cpp
--- a/Engine/Plugins/Experimental/ControlFlows/Source/ControlFlows/Private/ControlFlow.cpp
+++ b/Engine/Plugins/Experimental/ControlFlows/Source/ControlFlows/Private/ControlFlow.cpp
@@ -75,7 +75,9 @@ void FControlFlow::HandleControlFlowNodeCompleted(TSharedRef<const FControlFlowN
 
 	SubFlowStack_ForDebugging.Add(SharedThis(this));
 
-	if (ensureMsgf(CurrentNode.IsValid() && &NodeCompleted.Get() == CurrentNode.Get(), TEXT("This will likely become a static_assert/compile-error")))
+	if (ensureMsgf(CurrentNode.IsValid(), TEXT("This will likely become a static_assert/compile-error")) &&
+		ensureMsgf(&NodeCompleted.Get() == CurrentNode.Get(), TEXT("This will likely become a static_assert/compile-error"))
+	   )
 	{
 		if (!NodeCompleted->bWasBoundOnExecution)
 		{
