Added function: "QueueDelay"

diff --git a/Engine/Plugins/Experimental/ControlFlows/Source/ControlFlows/Private/ControlFlow.cpp b/Engine/Plugins/Experimental/ControlFlows/Source/ControlFlows/Private/ControlFlow.cpp
--- a/Engine/Plugins/Experimental/ControlFlows/Source/ControlFlows/Private/ControlFlow.cpp
+++ b/Engine/Plugins/Experimental/ControlFlows/Source/ControlFlows/Private/ControlFlow.cpp
@@ -287,6 +287,23 @@ TSharedPtr<FTrackedActivity> FControlFlow::GetTrackedActivity() const
 	return Activity;
 }
 
+FControlFlow& FControlFlow::QueueDelay(const float InDelay, const FString& NodeName)
+{
+	QueueWait(NodeName).BindLambda([InDelay](FControlFlowNodeRef FlowHandle)
+		{
+			FTSTicker::GetCoreTicker().AddTicker(FTickerDelegate::CreateLambda([FlowHandle](float)
+				{
+					if (FlowHandle->Parent.IsValid() && !FlowHandle->HasCancelBeenRequested())
+					{
+						FlowHandle->ContinueFlow();
+					}
+					return false;
+				}), InDelay);
+		});
+
+	return *this;
+}
+
 FControlFlow& FControlFlow::TrackActivities(TSharedPtr<FTrackedActivity> InActivity)
 {
 	if (!InActivity)
diff --git a/Engine/Plugins/Experimental/ControlFlows/Source/ControlFlows/Public/ControlFlow.h b/Engine/Plugins/Experimental/ControlFlows/Source/ControlFlows/Public/ControlFlow.h
--- a/Engine/Plugins/Experimental/ControlFlows/Source/ControlFlows/Public/ControlFlow.h
+++ b/Engine/Plugins/Experimental/ControlFlows/Source/ControlFlows/Public/ControlFlow.h
@@ -97,7 +97,7 @@ class FControlFlowTask_BranchLegacy;
  *												.QueueStep(this, &ThisClass::Foo);
  *										});
  * 
- *									return FMath::RandBool() ? 0 : 1;
+ *									return FMath::RandBool() ? EConditionalLoopResult::RunLoop : EConditionalLoopResult::LoopFinished;
  *								})
  *								.QueueStep(this, &ThisClass::Foo);
  *
@@ -148,6 +148,7 @@ public:
 	CONTROLFLOWS_API TSharedPtr<FTrackedActivity> GetTrackedActivity() const;
 
 public:
+	CONTROLFLOWS_API FControlFlow& QueueDelay(const float InDelay, const FString& NodeName = FString());
 
 	template<typename...ArgsT>
 	FControlFlow& QueueStep(const FString& NodeName, ArgsT...Params)
