Added the current and complete node names to an ensure that's supposed to show the nodes are the same to help track down what's happening.

Made FControlFlowNode::GetNodeName() public so that it can be used to give more info for debugging.

#UE5

diff --git a/Engine/Plugins/Experimental/ControlFlows/Source/ControlFlows/Private/ControlFlow.cpp b/Engine/Plugins/Experimental/ControlFlows/Source/ControlFlows/Private/ControlFlow.cpp
--- a/Engine/Plugins/Experimental/ControlFlows/Source/ControlFlows/Private/ControlFlow.cpp
+++ b/Engine/Plugins/Experimental/ControlFlows/Source/ControlFlows/Private/ControlFlow.cpp
@@ -76,8 +76,12 @@ void FControlFlow::HandleControlFlowNodeCompleted(TSharedRef<const FControlFlowN
 
 	SubFlowStack_ForDebugging.Add(SharedThis(this));
 
-	if (ensureMsgf(CurrentNode.IsValid(), TEXT("This will likely become a static_assert/compile-error")) &&
-		ensureMsgf(&NodeCompleted.Get() == CurrentNode.Get(), TEXT("This will likely become a static_assert/compile-error"))
+	if (ensureMsgf(CurrentNode.IsValid(), TEXT("CurrentNode isn't valid when handling control flow node (%s) being completed.  This will likely become a static_assert/compile-error"),
+			*NodeCompleted->GetNodeName()
+				  ) &&
+		ensureMsgf(&NodeCompleted.Get() == CurrentNode.Get(), TEXT("NodeCompleted (%s) does not match the current node (%s).  This will likely become a static_assert/compile-error"),
+			*NodeCompleted->GetNodeName(), *CurrentNode->GetNodeName()
+				  )
 	   )
 	{
 		if (!NodeCompleted->bWasBoundOnExecution)
