Added ability to queue interpreting cancelled nodes as completed, rather than just setting it on the control flow itself.

diff --git a/Engine/Plugins/Experimental/ControlFlows/Source/ControlFlows/Private/ControlFlow.cpp b/Engine/Plugins/Experimental/ControlFlows/Source/ControlFlows/Private/ControlFlow.cpp
--- a/Engine/Plugins/Experimental/ControlFlows/Source/ControlFlows/Private/ControlFlow.cpp
+++ b/Engine/Plugins/Experimental/ControlFlows/Source/ControlFlows/Private/ControlFlow.cpp
@@ -317,6 +317,16 @@ FControlFlow& FControlFlow::QueueDelay(const float InDelay, const FString& NodeN
 	return *this;
 }
 
+FControlFlow& FControlFlow::QueueSetCancelledNodeAsComplete(const bool bCancelledNodeIsComplete, const FString& NodeName /*= FString()*/)
+{
+	QueueFunction(NodeName).BindSPLambda(this, [this, bCancelledNodeIsComplete]()
+		{
+			this->SetCancelledNodeAsComplete(bCancelledNodeIsComplete);
+		});
+
+	return *this;
+}
+
 FControlFlow& FControlFlow::TrackActivities(TSharedPtr<FTrackedActivity> InActivity)
 {
 	if (!InActivity)
diff --git a/Engine/Plugins/Experimental/ControlFlows/Source/ControlFlows/Public/ControlFlow.h b/Engine/Plugins/Experimental/ControlFlows/Source/ControlFlows/Public/ControlFlow.h
--- a/Engine/Plugins/Experimental/ControlFlows/Source/ControlFlows/Public/ControlFlow.h
+++ b/Engine/Plugins/Experimental/ControlFlows/Source/ControlFlows/Public/ControlFlow.h
@@ -149,6 +149,7 @@ public:
 
 public:
 	CONTROLFLOWS_API FControlFlow& QueueDelay(const float InDelay, const FString& NodeName = FString());
+	CONTROLFLOWS_API FControlFlow& QueueSetCancelledNodeAsComplete(const bool bCancelledNodeIsComplete, const FString& NodeName = FString());
 
 	template<typename...ArgsT>
 	FControlFlow& QueueStep(const FString& NodeName, ArgsT...Params)
