Moving "CastedTo.h" and "SmartPtr.h" to Experimental plugin

diff --git a/Engine/Plugins/Experimental/ControlFlows/Source/ControlFlows/Public/CastedTo.h b/Engine/Plugins/Experimental/ControlFlows/Source/ControlFlows/Public/CastedTo.h
new file mode 100644
--- /dev/null
+++ b/Engine/Plugins/Experimental/ControlFlows/Source/ControlFlows/Public/CastedTo.h
@@ -0,0 +1,163 @@
+// Copyright Epic Games, Inc. All Rights Reserved.
+
+#pragma once
+
+#include "UObject/StrongObjectPtr.h"
+#include "UObject/WeakObjectPtr.h"
+#include "Templates/SharedPointer.h"
+
+/* A Pointer Container for a 'T' (UObject xOr TSharedFromThis) that can be statically casted to 'CastTo'
+ * These classes are useful if you want your system to support both memory management systems: (1) UObject Garbage Collection and (2) Smart Pointer Reference counting
+ */
+template<typename CastTo>
+class TPointerContainerBase : public TSharedFromThis<TPointerContainerBase<CastTo>>
+{
+public:
+	virtual ~TPointerContainerBase() {}
+
+	virtual CastTo* Cast() const = 0;
+	virtual bool IsValid() const = 0;
+	virtual void Reset() = 0;
+
+public:
+	CastTo* Get() const { return Cast(); }
+	bool IsPtrValid() const { return IsValid(); }
+	void ResetPtr() { Reset(); }
+};
+
+template<typename CastTo>
+class TPointerContainer : public TPointerContainerBase<CastTo>
+{
+public:
+	virtual ~TPointerContainer() {}
+
+	virtual CastTo* Cast() const override { return nullptr; }
+	virtual bool IsValid() const override { return false; }
+	virtual void Reset() override {}
+};
+
+/* A "Weak" Pointer Reference Container for a 'T' (UObject xOr TSharedFromThis) that can be statically casted to 'CastTo' */
+template<typename CastTo>
+class TWeakContainerTo : public TPointerContainer<CastTo>
+{
+public:
+	virtual ~TWeakContainerTo() {}
+};
+
+template<typename T, typename CastTo, bool /*bDerivedFromUObject*/, bool /*bDerivedFromTSharedFromThis*/>
+class TWeakCastableImpl;
+
+template<typename T, typename CastTo>
+class TWeakCastableImpl<T, CastTo, true, false> : public TWeakContainerTo<CastTo>
+{
+public:
+	TWeakCastableImpl() {}
+	TWeakCastableImpl(T* InObject) : Object(InObject) {}
+	virtual ~TWeakCastableImpl() {}
+
+public:
+	virtual CastTo* Cast() const override final { return static_cast<CastTo*>(Object.Get()); }
+	virtual bool IsValid() const override final { return Object.IsValid(); }
+	virtual void Reset() override final { Object.Reset(); }
+
+private:
+	TWeakObjectPtr<T> Object;
+};
+
+template<typename T, typename CastTo>
+class TWeakCastableImpl<T, CastTo, false, true> : public TWeakContainerTo<CastTo>
+{
+public:
+	TWeakCastableImpl() {}
+	TWeakCastableImpl(T* InObject) : Object(StaticCastSharedRef<T>(InObject->AsShared())) {}
+	virtual ~TWeakCastableImpl() {}
+
+public:
+	virtual CastTo* Cast() const override final { return static_cast<CastTo*>(IsValid() ? Object.Pin().Get() : nullptr); }
+	virtual bool IsValid() const override final { return Object.IsValid(); }
+	virtual void Reset() override final { Object.Reset(); }
+
+private:
+	TWeakPtr<T> Object;
+};
+
+#define TWeakCastableImpl_Decl TWeakCastableImpl<CastFrom, CastTo, TIsDerivedFrom<CastFrom, UObject>::IsDerived, IsDerivedFromSharedFromThis<CastFrom>()>
+template<typename CastFrom, typename CastTo>
+class TWeakCastable : public TWeakCastableImpl_Decl
+{
+public:
+	TWeakCastable() {}
+	TWeakCastable(CastFrom* InObject) : TWeakCastableImpl_Decl(InObject) {}
+};
+#undef TWeakCastableImpl_Decl
+
+template<typename T>
+class TWeakContainer : public TWeakCastable<T, T>
+{
+public:
+	TWeakContainer() : TWeakCastable<T, T>() {}
+	TWeakContainer(T* InObject) : TWeakCastable<T, T>(InObject) {}
+};
+
+/* A "Strong" Pointer Reference Container for a 'T' (UObject xOr TSharedFromThis) that can be statically casted to 'CastTo' */
+template<typename CastTo>
+class TStrongContainerTo : public TPointerContainer<CastTo>
+{
+public:
+	virtual ~TStrongContainerTo() {}
+};
+
+template<typename T, typename CastTo, bool /*bDerivedFromUObject*/, bool /*bDerivedFromTSharedFromThis*/>
+class TStrongContainerToImpl;
+
+template<typename T, typename CastTo>
+class TStrongContainerToImpl<T, CastTo, true, false> : public TStrongContainerTo<CastTo>
+{
+public:
+	TStrongContainerToImpl() {}
+	TStrongContainerToImpl(T* InObject) : Object(InObject) {}
+	virtual ~TStrongContainerToImpl() {}
+
+public:
+	virtual CastTo* Cast() const override final { return static_cast<CastTo*>(Object.Get()); }
+	virtual bool IsValid() const override final { return Object.IsValid(); }
+	virtual void Reset() override final { Object.Reset(); }
+
+private:
+	TStrongObjectPtr<T> Object;
+};
+
+template<typename T, typename CastTo>
+class TStrongContainerToImpl<T, CastTo, false, true> : public TStrongContainerTo<CastTo>
+{
+public:
+	TStrongContainerToImpl() {}
+	TStrongContainerToImpl(T* InObject) : Object(StaticCastSharedRef<T>(InObject->AsShared())) {}
+	virtual ~TStrongContainerToImpl() {}
+
+public:
+	virtual CastTo* Cast() const override final { return static_cast<CastTo*>(Object.IsValid() ? Object.Get() : nullptr); }
+	virtual bool IsValid() const override final { return Object.IsValid(); }
+	virtual void Reset() override final { Object.Reset(); }
+
+private:
+	TSharedPtr<T> Object;
+};
+
+#define TStrongContainerToImpl_Decl TStrongContainerToImpl<CastFrom, CastTo, TIsDerivedFrom<CastFrom, UObject>::IsDerived, IsDerivedFromSharedFromThis<CastFrom>()>
+template<typename CastFrom, typename CastTo>
+class TStrongCastable : public TStrongContainerToImpl_Decl
+{
+public:
+	TStrongCastable() : TStrongContainerToImpl_Decl() {}
+	TStrongCastable(CastFrom* InObject) : TStrongContainerToImpl_Decl(InObject) {}
+};
+#undef TStrongContainerToImpl_Decl
+
+template<typename T>
+class TStrongContainer : public TStrongCastable<T, T>
+{
+public:
+	TStrongContainer() : TStrongCastable<T, T>() {}
+	TStrongContainer(T* InObject) : TStrongCastable<T, T>(InObject) {}
+};
\ No newline at end of file
diff --git a/Engine/Plugins/Experimental/ControlFlows/Source/ControlFlows/Public/ControlFlowContainer.h b/Engine/Plugins/Experimental/ControlFlows/Source/ControlFlows/Public/ControlFlowContainer.h
--- a/Engine/Plugins/Experimental/ControlFlows/Source/ControlFlows/Public/ControlFlowContainer.h
+++ b/Engine/Plugins/Experimental/ControlFlows/Source/ControlFlows/Public/ControlFlowContainer.h
@@ -5,7 +5,7 @@
 // HEADER_UNIT_SKIP - Not included directly
 
 #include "Templates/SharedPointer.h"
-#include "Containers/CastedTo.h"
+#include "CastedTo.h"
 
 class FControlFlowContainerBase : public TSharedFromThis<FControlFlowContainerBase>
 {
diff --git a/Engine/Plugins/Experimental/ControlFlows/Source/ControlFlows/Public/SmartPtr.h b/Engine/Plugins/Experimental/ControlFlows/Source/ControlFlows/Public/SmartPtr.h
new file mode 100644
--- /dev/null
+++ b/Engine/Plugins/Experimental/ControlFlows/Source/ControlFlows/Public/SmartPtr.h
@@ -0,0 +1,63 @@
+// Copyright Epic Games, Inc. All Rights Reserved.
+
+#pragma once
+
+#include "CastedTo.h"
+
+enum class ESmartPointer
+{
+	Strong,
+	Shared,
+	Weak,
+	// Unique(?) - not sure if needed
+};
+
+/* A container to hold one of the FOUR Smart Pointer types: TStrongObjectPtr, TWeakObjectPtr, TSharedPtr, or TWeakPtr. 
+* During construction, there will be a compile time attempt to static_cast a raw pointer 'From' to 'T'
+* 'From' must derive from 'UObject' xOr 'TSharedFromThis' for this class to be usable
+*/
+template<typename T>
+class TSmartPtr
+{
+public:
+	TSmartPtr() : Container(MakeShared<TPointerContainer<T>>()) {}
+
+	template<typename From>
+	TSmartPtr(From* InObject)
+	{
+		Set(InObject, ESmartPointer::Weak);
+	}
+
+	template<typename From>
+	TSmartPtr(From* InObject, const ESmartPointer InType)
+	{
+		Set(InObject, InType);
+	}
+
+	template<typename From>
+	TSmartPtr<T>& Set(From* InObject, const ESmartPointer InType)
+	{
+		if (InType == ESmartPointer::Strong || InType == ESmartPointer::Shared)
+		{
+			ensureAlwaysMsgf(
+				(TIsDerivedFrom<From, UObject>::IsDerived && InType == ESmartPointer::Strong) ||
+				(IsDerivedFromSharedFromThis<From>() && InType == ESmartPointer::Shared),
+				TEXT("Mismatch of Strong and Shared. We can continue on holding a 'strong' reference, but you shouldn't mix UObjects and Shared Pointers"));
+
+			Container = MakeShared<TStrongCastable<From, T>>(InObject);
+		}
+		else if (InType == ESmartPointer::Weak)
+		{
+			Container = MakeShared<TWeakCastable<From, T>>(InObject);
+		}
+
+		return *this;
+	}
+
+	T* Get() const { return Container->Get(); }
+	bool IsValid() const { return Container->IsValid(); }
+	void Reset() const { Container->Reset(); }
+
+private:
+	TSharedRef<TPointerContainer<T>> Container = MakeShared<TPointerContainer<T>>();
+};
\ No newline at end of file
diff --git a/Engine/Source/Runtime/Core/Public/Containers/CastedTo.h b/Engine/Source/Runtime/Core/Public/Containers/CastedTo.h
deleted file mode 100644
diff --git a/Engine/Source/Runtime/Core/Public/Templates/SmartPtr.h b/Engine/Source/Runtime/Core/Public/Templates/SmartPtr.h
deleted file mode 100644
