Crash fix. Changed from deriving to a member variable

diff --git a/Engine/Plugins/Experimental/ControlFlows/Source/ControlFlows/Public/ControlFlowContainer.h b/Engine/Plugins/Experimental/ControlFlows/Source/ControlFlows/Public/ControlFlowContainer.h
--- a/Engine/Plugins/Experimental/ControlFlows/Source/ControlFlows/Public/ControlFlowContainer.h
+++ b/Engine/Plugins/Experimental/ControlFlows/Source/ControlFlows/Public/ControlFlowContainer.h
@@ -38,18 +38,21 @@ private:
 };
 
 template<typename T>
-class TControlFlowContainer : public FControlFlowContainerBase, public TWeakContainer<T>
+class TControlFlowContainer : public FControlFlowContainerBase
 {
 public:
 	TControlFlowContainer() = delete;
 	TControlFlowContainer(T* InOwner, TSharedRef<FControlFlow> InFlow, const FString& FlowId)
 		: FControlFlowContainerBase(InFlow, FlowId)
-		, TWeakContainer<T>(InOwner)
+		, OwningObject(InOwner)
 	{}
 
 private:
 	virtual const void* const GetOwningObject() const override final
 	{
-		return static_cast<const TWeakContainerTo<T>*>(this)->Cast();
+		return OwningObject.IsValid() ? OwningObject.Cast() : nullptr;
 	}
+
+private:
+	TWeakContainer<T> OwningObject;
 };
\ No newline at end of file
