Fixed bug with multiple QueueDelay(0.0f) in the same FlowQueue

diff --git a/Engine/Plugins/Experimental/ControlFlows/Source/ControlFlows/Private/ControlFlow.cpp b/Engine/Plugins/Experimental/ControlFlows/Source/ControlFlows/Private/ControlFlow.cpp
--- a/Engine/Plugins/Experimental/ControlFlows/Source/ControlFlows/Private/ControlFlow.cpp
+++ b/Engine/Plugins/Experimental/ControlFlows/Source/ControlFlows/Private/ControlFlow.cpp
@@ -6,6 +6,7 @@
 #include "ControlFlowTask_Loop.h"
 #include "ControlFlows.h"
 #include "Containers/Ticker.h"
+#include "Misc/App.h"
 #include "Misc/TrackedActivity.h"
 
 static const int32 MAX_FLOW_LOOPS = 10000;
@@ -300,18 +301,36 @@ TSharedPtr<FTrackedActivity> FControlFlow::GetTrackedActivity() const
 	return Activity;
 }
 
-FControlFlow& FControlFlow::QueueDelay(const float InDelay, const FString& NodeName)
+FControlFlow& FControlFlow::QueueDelay(const float InSeconds, const FString& NodeName)
 {
-	QueueWait(NodeName).BindLambda([InDelay](FControlFlowNodeRef FlowHandle)
+	QueueWait(NodeName).BindLambda([InSeconds](FControlFlowNodeRef FlowHandle)
 		{
-			FTSTicker::GetCoreTicker().AddTicker(FTickerDelegate::CreateLambda([FlowHandle](float)
+			TSharedRef<FControlFlow> OwningFlow = FlowHandle->Parent.Pin().ToSharedRef();
+
+			// If we already had a 0.0f delay this tick, add the delta time of this tick to the delay
+			const float AdditionalDelay = (InSeconds == 0.0f && OwningFlow->LastZeroSecondDelay.Key == FApp::GetCurrentTime()) ? OwningFlow->LastZeroSecondDelay.Value : 0.0f;
+			const float SecondsToDelay = InSeconds + AdditionalDelay;
+
+			FTSTicker::GetCoreTicker().AddTicker(FTickerDelegate::CreateLambda([FlowHandle, bIsZeroSecondDelay = SecondsToDelay == 0.0f](float DeltaTime)
 				{
 					if (FlowHandle->Parent.IsValid() && !FlowHandle->HasCancelBeenRequested())
 					{
+						// Notify ALL parents of 0.0f delay
+						if (bIsZeroSecondDelay)
+						{
+							TSharedPtr<FControlFlow> FlowIter = FlowHandle->Parent.Pin().ToSharedRef();
+							while (FlowIter.IsValid())
+							{
+								FlowIter->LastZeroSecondDelay = { FApp::GetCurrentTime(), DeltaTime };
+
+								FlowIter = FlowIter->ParentFlow.IsValid() ? FlowIter->ParentFlow.Pin() : nullptr;
+							}
+						}
+
 						FlowHandle->ContinueFlow();
 					}
 					return false;
-				}), InDelay);
+				}), SecondsToDelay);
 		});
 
 	return *this;
@@ -353,8 +372,12 @@ FControlFlowWaitDelegate& FControlFlow::QueueWait(const FString& FlowNodeDebugNa
 
 FControlFlowPopulator& FControlFlow::QueueControlFlow(const FString& TaskName /*= TEXT("")*/, const FString& FlowNodeDebugName /*= TEXT("")*/)
 {
-	TSharedRef<FControlFlowSimpleSubTask> NewTask = MakeShared<FControlFlowSimpleSubTask>(TaskName, MakeShared<FControlFlow>(TaskName));
+	TSharedRef<FControlFlowSimpleSubTask> NewTask = MakeShared<FControlFlowSimpleSubTask>(TaskName);
 	TSharedRef<FControlFlowNode_Task> NewNode = MakeShared<FControlFlowNode_Task>(SharedThis(this), NewTask, FormatOrGetNewNodeDebugName(FlowNodeDebugName));
+
+	NewTask->GetTaskFlow()->ParentFlow = AsWeak();
+
+	NewTask->SetOwningNode(NewNode);
 	NewTask->GetTaskFlow()->Activity = Activity;
 	NewNode->OnExecute().BindSP(SharedThis(this), &FControlFlow::HandleTaskNodeExecuted);
 	NewNode->OnCancelRequested().BindSP(SharedThis(this), &FControlFlow::HandleTaskNodeCancelled);
@@ -369,6 +392,7 @@ FControlFlowBranchDefiner& FControlFlow::QueueControlFlowBranch(const FString& T
 	TSharedRef<FControlFlowTask_Branch> NewTask = MakeShared<FControlFlowTask_Branch>(TaskName);
 	TSharedRef<FControlFlowNode_Task> NewNode = MakeShared<FControlFlowNode_Task>(SharedThis(this), NewTask, FormatOrGetNewNodeDebugName(FlowNodeDebugName));
 
+	NewTask->SetOwningNode(NewNode);
 	NewTask->Activity = Activity;
 	NewNode->OnExecute().BindSP(SharedThis(this), &FControlFlow::HandleTaskNodeExecuted);
 	NewNode->OnCancelRequested().BindSP(SharedThis(this), &FControlFlow::HandleTaskNodeCancelled);
@@ -383,6 +407,7 @@ FConcurrentFlowsDefiner& FControlFlow::QueueConcurrentFlows(const FString& TaskN
 	TSharedRef<FControlFlowTask_ConcurrentFlows> NewTask = MakeShared<FControlFlowTask_ConcurrentFlows>(TaskName);
 	TSharedRef<FControlFlowNode_Task> NewNode = MakeShared<FControlFlowNode_Task>(SharedThis(this), NewTask, FormatOrGetNewNodeDebugName(FlowNodeDebugName));
 
+	NewTask->SetOwningNode(NewNode);
 	NewNode->OnExecute().BindSP(SharedThis(this), &FControlFlow::HandleTaskNodeExecuted);
 	NewNode->OnCancelRequested().BindSP(SharedThis(this), &FControlFlow::HandleTaskNodeCancelled);
 
@@ -396,6 +421,7 @@ FControlFlowConditionalLoopDefiner& FControlFlow::QueueConditionalLoop(const FSt
 	TSharedRef<FControlFlowTask_ConditionalLoop> NewTask = MakeShared<FControlFlowTask_ConditionalLoop>(TaskName);
 	TSharedRef<FControlFlowNode_Task> NewNode = MakeShared<FControlFlowNode_Task>(SharedThis(this), NewTask, FormatOrGetNewNodeDebugName(FlowNodeDebugName));
 
+	NewTask->SetOwningNode(NewNode);
 	NewNode->OnExecute().BindSP(SharedThis(this), &FControlFlow::HandleTaskNodeExecuted);
 	NewNode->OnCancelRequested().BindSP(SharedThis(this), &FControlFlow::HandleTaskNodeCancelled);
 
@@ -411,6 +437,7 @@ TSharedRef<FControlFlowTask_BranchLegacy> FControlFlow::QueueBranch(FControlFlow
 	TSharedRef<FControlFlowTask_BranchLegacy> NewTask = MakeShared<FControlFlowTask_BranchLegacy>(BranchDecider, TaskName);
 	TSharedRef<FControlFlowNode_Task> NewNode = MakeShared<FControlFlowNode_Task>(SharedThis(this), NewTask, FormatOrGetNewNodeDebugName(FlowNodeDebugName));
 
+	NewTask->SetOwningNode(NewNode);
 	NewNode->OnExecute().BindSP(SharedThis(this), &FControlFlow::HandleTaskNodeExecuted);
 	NewNode->OnCancelRequested().BindSP(SharedThis(this), &FControlFlow::HandleTaskNodeCancelled);
 
@@ -421,8 +448,12 @@ TSharedRef<FControlFlowTask_BranchLegacy> FControlFlow::QueueBranch(FControlFlow
 
 FControlFlowPopulator& FControlFlow::QueueLoop(FControlFlowLoopComplete& LoopCompleteDelgate, const FString& TaskName /*= TEXT("")*/, const FString& FlowNodeDebugName /*= TEXT("")*/)
 {
-	TSharedRef<FControlFlowTask_LoopDeprecated> NewTask = MakeShared<FControlFlowTask_LoopDeprecated>(LoopCompleteDelgate, TaskName, MakeShared<FControlFlow>(TaskName));
+	TSharedRef<FControlFlowTask_LoopDeprecated> NewTask = MakeShared<FControlFlowTask_LoopDeprecated>(LoopCompleteDelgate, TaskName);
 	TSharedRef<FControlFlowNode_Task> NewNode = MakeShared<FControlFlowNode_Task>(SharedThis(this), NewTask, FormatOrGetNewNodeDebugName(FlowNodeDebugName));
+
+	NewTask->GetTaskFlow()->ParentFlow = AsWeak();
+
+	NewTask->SetOwningNode(NewNode);
 	NewTask->GetTaskFlow()->Activity = Activity;
 	NewNode->OnExecute().BindSP(SharedThis(this), &FControlFlow::HandleTaskNodeExecuted);
 	NewNode->OnCancelRequested().BindSP(SharedThis(this), &FControlFlow::HandleTaskNodeCancelled);
diff --git a/Engine/Plugins/Experimental/ControlFlows/Source/ControlFlows/Private/ControlFlowConcurrency.cpp b/Engine/Plugins/Experimental/ControlFlows/Source/ControlFlows/Private/ControlFlowConcurrency.cpp
--- a/Engine/Plugins/Experimental/ControlFlows/Source/ControlFlows/Private/ControlFlowConcurrency.cpp
+++ b/Engine/Plugins/Experimental/ControlFlows/Source/ControlFlows/Private/ControlFlowConcurrency.cpp
@@ -3,6 +3,7 @@
 #include "ControlFlowConcurrency.h"
 #include "ControlFlow.h"
 #include "ControlFlows.h"
+#include "ControlFlowTask.h"
 
 FConcurrencySubFlowContainer::FConcurrencySubFlowContainer(const FString& InDebugName)
 	: SubFlow(MakeShared<FControlFlow>(InDebugName))
@@ -86,6 +87,10 @@ FControlFlow& FConcurrentControlFlows::AddOrGetFlow(int32 InIdentifier, const FS
 
 		TSharedRef<FConcurrencySubFlowContainer> NewConcurrentFlow = MakeShared<FConcurrencySubFlowContainer>(DebugNameToUse);
 
+		check(OwningTask.IsValid());
+
+		NewConcurrentFlow->GetControlFlow()->ParentFlow = OwningTask.Pin()->GetOwningFlowForTaskNode();
+
 		NewConcurrentFlow->OnComplete().BindSP(SharedThis(this), &FConcurrentControlFlows::HandleConcurrentFlowCompleted, InIdentifier);
 		NewConcurrentFlow->OnExecutedWithoutAnyNodes().BindSP(SharedThis(this), &FConcurrentControlFlows::HandleConcurrentFlowCompleted, InIdentifier);
 		NewConcurrentFlow->OnCancelled().BindSP(SharedThis(this), &FConcurrentControlFlows::HandleConcurrentFlowCancelled, InIdentifier);
@@ -175,6 +180,8 @@ void FConcurrentControlFlows::Execute()
 		UE_LOG(LogControlFlows, Verbose, TEXT("ConcurrentControlFlow::Execute - Executing Subflow %s"), *PairIt.Value->GetDebugName());
 
 		bFlowExecuted = true;
+
+		PairIt.Value->GetControlFlow()->LastZeroSecondDelay = OwningTask.Pin()->GetOwningFlowForTaskNode().Pin()->LastZeroSecondDelay;
 		PairIt.Value->Execute();
 
 		if (bCancelAllHasBegun)
diff --git a/Engine/Plugins/Experimental/ControlFlows/Source/ControlFlows/Private/ControlFlowNode.cpp b/Engine/Plugins/Experimental/ControlFlows/Source/ControlFlows/Private/ControlFlowNode.cpp
--- a/Engine/Plugins/Experimental/ControlFlows/Source/ControlFlows/Private/ControlFlowNode.cpp
+++ b/Engine/Plugins/Experimental/ControlFlows/Source/ControlFlows/Private/ControlFlowNode.cpp
@@ -130,11 +130,10 @@ void FControlFlowNode_RequiresCallback::Execute()
 
 ///////////////////////////////////////////////////
 
-FControlFlowNode_Task::FControlFlowNode_Task(TSharedRef<FControlFlow> ControlFlowParent, TSharedRef<FControlFlowSubTaskBase> Module, const FString& FlowNodeDebugName)
+FControlFlowNode_Task::FControlFlowNode_Task(TSharedRef<FControlFlow> ControlFlowParent, TSharedRef<FControlFlowSubTaskBase> ControlFlowTask, const FString& FlowNodeDebugName)
 	: FControlFlowNode(ControlFlowParent, FlowNodeDebugName)
-	, FlowTask(Module)
+	, FlowTask(ControlFlowTask)
 {
-
 }
 
 void FControlFlowNode_Task::Execute()
@@ -167,4 +166,4 @@ FString FControlFlowNode_Task::GetNodeName() const
 void FControlFlowNode_Task::CompleteCancelFlow()
 {
 	FControlFlowNode::CancelFlow();
-}
+}
\ No newline at end of file
diff --git a/Engine/Plugins/Experimental/ControlFlows/Source/ControlFlows/Private/ControlFlowTask.cpp b/Engine/Plugins/Experimental/ControlFlows/Source/ControlFlows/Private/ControlFlowTask.cpp
--- a/Engine/Plugins/Experimental/ControlFlows/Source/ControlFlows/Private/ControlFlowTask.cpp
+++ b/Engine/Plugins/Experimental/ControlFlows/Source/ControlFlows/Private/ControlFlowTask.cpp
@@ -15,6 +15,11 @@ FControlFlowSubTaskBase::FControlFlowSubTaskBase(const FString& TaskName)
 
 }
 
+TWeakPtr<FControlFlow> FControlFlowSubTaskBase::GetOwningFlowForTaskNode() const
+{
+	return ensureAlways(OwningNode.IsValid()) ? OwningNode.Pin()->Parent : nullptr;
+}
+
 void FControlFlowSubTaskBase::Execute()
 {
 	TaskCompleteCallback.ExecuteIfBound();
@@ -79,7 +84,11 @@ TSharedRef<FControlFlow> FControlFlowTask_BranchLegacy::GetOrAddBranch(int32 Bra
 {
 	if (!Branches.Contains(BranchIndex))
 	{
-		Branches.Add(BranchIndex, MakeShared<FControlFlow>());
+		TSharedRef<FControlFlow> NewBranchFlow = MakeShared<FControlFlow>();
+
+		NewBranchFlow->ParentFlow = GetOwningFlowForTaskNode();
+		
+		Branches.Add(BranchIndex, NewBranchFlow);
 	}
 
 	return Branches[BranchIndex];
@@ -91,7 +100,9 @@ void FControlFlowTask_BranchLegacy::Execute()
 	{
 		SelectedBranch = BranchDelegate.Execute();
 
-		TSharedRef<FControlFlow> FlowToExecute = GetOrAddBranch(SelectedBranch);
+		TSharedRef<FControlFlow> FlowToExecute = GetOrAddBranch(*SelectedBranch);
+
+		FlowToExecute->LastZeroSecondDelay = GetOwningFlowForTaskNode().Pin()->LastZeroSecondDelay;
 
 		FlowToExecute->OnCompleteDelegate_Internal.BindSP(SharedThis(this), &FControlFlowTask_BranchLegacy::HandleBranchCompleted);
 		FlowToExecute->OnExecutedWithoutAnyNodesDelegate_Internal.BindSP(SharedThis(this), &FControlFlowTask_BranchLegacy::HandleBranchCompleted);
@@ -107,9 +118,9 @@ void FControlFlowTask_BranchLegacy::Execute()
 
 void FControlFlowTask_BranchLegacy::Cancel()
 {
-	if (Branches.Contains(SelectedBranch) && Branches[SelectedBranch]->IsRunning())
+	if (ensureAlways(SelectedBranch.IsSet()) && Branches.Contains(*SelectedBranch) && Branches[*SelectedBranch]->IsRunning())
 	{
-		Branches[SelectedBranch]->CancelFlow();
+		Branches[*SelectedBranch]->CancelFlow();
 	}
 	else
 	{
@@ -121,17 +132,18 @@ void FControlFlowTask_BranchLegacy::Cancel()
 //FControlFlowSimpleSubTask
 //////////////////////////////////
 
-FControlFlowSimpleSubTask::FControlFlowSimpleSubTask(const FString& TaskName, TSharedRef<FControlFlow> FlowOwner)
+FControlFlowSimpleSubTask::FControlFlowSimpleSubTask(const FString& TaskName)
 	: FControlFlowSubTaskBase(TaskName)
-	, TaskFlow(FlowOwner)
+	, TaskFlow(MakeShared<FControlFlow>(TaskName))
 {
-
 }
 
 void FControlFlowSimpleSubTask::Execute()
 {
 	if (TaskPopulator.IsBound() && GetTaskFlow().IsValid())
 	{
+		GetTaskFlow()->LastZeroSecondDelay = GetOwningFlowForTaskNode().Pin()->LastZeroSecondDelay;
+
 		GetTaskFlow()->OnCompleteDelegate_Internal.BindSP(SharedThis(this), &FControlFlowSimpleSubTask::CompletedSubTask);
 		GetTaskFlow()->OnExecutedWithoutAnyNodesDelegate_Internal.BindSP(SharedThis(this), &FControlFlowSimpleSubTask::CompletedSubTask);
 		GetTaskFlow()->OnCancelledDelegate_Internal.BindSP(SharedThis(this), &FControlFlowSimpleSubTask::CancelledSubTask);
@@ -176,8 +188,8 @@ void FControlFlowSimpleSubTask::CancelledSubTask()
 //FControlFlowLoop
 //////////////////////////////////
 
-FControlFlowTask_LoopDeprecated::FControlFlowTask_LoopDeprecated(FControlFlowLoopComplete& TaskCompleteDelegate, const FString& TaskName, TSharedRef<FControlFlow> FlowOwner)
-	: FControlFlowSimpleSubTask(TaskName, FlowOwner)
+FControlFlowTask_LoopDeprecated::FControlFlowTask_LoopDeprecated(FControlFlowLoopComplete& TaskCompleteDelegate, const FString& TaskName)
+	: FControlFlowSimpleSubTask(TaskName)
 	, TaskCompleteDecider(TaskCompleteDelegate)
 {}
 
@@ -191,6 +203,8 @@ void FControlFlowTask_LoopDeprecated::Execute()
 		}
 		else
 		{
+			GetTaskFlow()->LastZeroSecondDelay = GetOwningFlowForTaskNode().Pin()->LastZeroSecondDelay;
+
 			GetTaskFlow()->OnCompleteDelegate_Internal.BindSP(SharedThis(this), &FControlFlowTask_LoopDeprecated::CompletedLoop);
 			GetTaskFlow()->OnExecutedWithoutAnyNodesDelegate_Internal.BindSP(SharedThis(this), &FControlFlowTask_LoopDeprecated::CompletedLoop);
 			GetTaskFlow()->OnCancelledDelegate_Internal.BindSP(SharedThis(this), &FControlFlowTask_LoopDeprecated::CancelledLoop);
@@ -251,6 +265,10 @@ void FControlFlowTask_Branch::Execute()
 		if (ensureAlwaysMsgf(BranchDefinitions->Contains(SelectedBranchKey), TEXT("You've returned a Branch Key that doesn't exist!")))
 		{
 			SelectedBranchFlow = BranchDefinitions->FindChecked(SelectedBranchKey);
+
+			SelectedBranchFlow->ParentFlow = GetOwningFlowForTaskNode();
+			SelectedBranchFlow->LastZeroSecondDelay = GetOwningFlowForTaskNode().Pin()->LastZeroSecondDelay;
+
 			SelectedBranchFlow->Activity = Activity;
 
 			SelectedBranchFlow->OnCompleteDelegate_Internal.BindSP(SharedThis(this), &FControlFlowTask_Branch::HandleBranchCompleted);
@@ -317,6 +335,8 @@ void FControlFlowTask_ConcurrentFlows::Execute()
 	if (ensureAlways(ConcurrentFlowDelegate.IsBound() && !ConcurrentFlows.IsValid()))
 	{
 		ConcurrentFlows = MakeShared<FConcurrentControlFlows>();
+		ConcurrentFlows->OwningTask = AsWeak();
+
 		ConcurrentFlows->OnConcurrencyCompleted.BindSP(SharedThis(this), &FControlFlowTask_ConcurrentFlows::HandleConcurrentFlowsCompleted);
 		ConcurrentFlows->OnConcurrencyCancelled.BindSP(SharedThis(this), &FControlFlowTask_ConcurrentFlows::HandleConcurrentFlowsCancelled);
 
diff --git a/Engine/Plugins/Experimental/ControlFlows/Source/ControlFlows/Private/ControlFlowTask_Loop.cpp b/Engine/Plugins/Experimental/ControlFlows/Source/ControlFlows/Private/ControlFlowTask_Loop.cpp
--- a/Engine/Plugins/Experimental/ControlFlows/Source/ControlFlows/Private/ControlFlowTask_Loop.cpp
+++ b/Engine/Plugins/Experimental/ControlFlows/Source/ControlFlows/Private/ControlFlowTask_Loop.cpp
@@ -11,6 +11,8 @@ void FControlFlowTask_ConditionalLoop::Execute()
 	{
 		ConditionalLoop = MakeShared<FConditionalLoop>();
 
+		ConditionalLoop->FlowLoop->ParentFlow = GetOwningFlowForTaskNode();
+
 		ConditionalLoop->FlowLoop->OnCompleteDelegate_Internal.BindSP(SharedThis(this), &FControlFlowTask_ConditionalLoop::HandleLoopCompleted);
 		ConditionalLoop->FlowLoop->OnExecutedWithoutAnyNodesDelegate_Internal.BindSP(SharedThis(this), &FControlFlowTask_ConditionalLoop::HandleLoopCompleted);
 		ConditionalLoop->FlowLoop->OnCancelledDelegate_Internal.BindSP(SharedThis(this), &FControlFlowTask_ConditionalLoop::HandleLoopCancelled);
@@ -33,6 +35,7 @@ void FControlFlowTask_ConditionalLoop::Execute()
 				}
 				else
 				{
+					ConditionalLoop->FlowLoop->LastZeroSecondDelay = GetOwningFlowForTaskNode().Pin()->LastZeroSecondDelay;
 					ConditionalLoop->FlowLoop->ExecuteFlow();
 				}
 			}
@@ -63,6 +66,7 @@ void FControlFlowTask_ConditionalLoop::Helper_ConditionalCheck(const EConditiona
 {
 	if (InCondition == EConditionalLoopResult::RunLoop)
 	{
+		ConditionalLoop->FlowLoop->LastZeroSecondDelay = GetOwningFlowForTaskNode().Pin()->LastZeroSecondDelay;
 		ConditionalLoop->FlowLoop->ExecuteFlow();
 	}
 	else if (InCondition == EConditionalLoopResult::LoopFinished)
diff --git a/Engine/Plugins/Experimental/ControlFlows/Source/ControlFlows/Public/ControlFlow.h b/Engine/Plugins/Experimental/ControlFlows/Source/ControlFlows/Public/ControlFlow.h
--- a/Engine/Plugins/Experimental/ControlFlows/Source/ControlFlows/Public/ControlFlow.h
+++ b/Engine/Plugins/Experimental/ControlFlows/Source/ControlFlows/Public/ControlFlow.h
@@ -148,7 +148,7 @@ public:
 	CONTROLFLOWS_API TSharedPtr<FTrackedActivity> GetTrackedActivity() const;
 
 public:
-	CONTROLFLOWS_API FControlFlow& QueueDelay(const float InDelay, const FString& NodeName = FString());
+	CONTROLFLOWS_API FControlFlow& QueueDelay(const float InSeconds, const FString& NodeName = FString());
 	CONTROLFLOWS_API FControlFlow& QueueSetCancelledNodeAsComplete(const bool bCancelledNodeIsComplete, const FString& NodeName = FString());
 
 	template<typename...ArgsT>
@@ -328,6 +328,7 @@ private:
 	friend class FControlFlowTask_Branch;
 	friend class FControlFlowTask_ConditionalLoop;
 	friend class FControlFlowStatics;
+	friend class FConcurrentControlFlows;
 	friend struct FConcurrencySubFlowContainer;
 
 public:
@@ -391,6 +392,10 @@ private:
 
 	TSharedPtr<FControlFlowNode> CurrentNode = nullptr;
 
+	TWeakPtr<FControlFlow> ParentFlow;
+
+	TPair<double /*Timestamp*/, float /*DeltaTime*/> LastZeroSecondDelay;
+
 	//TODO: Put behind some args, because this is expensive.
 	TArray<TSharedRef<FControlFlow>> SubFlowStack_ForDebugging;
 
@@ -399,7 +404,6 @@ private:
 	TSharedPtr<FTrackedActivity> Activity;
 
 	bool bProfilerEventStarted = false;
-
 };
 
 #if CPUPROFILERTRACE_ENABLED
diff --git a/Engine/Plugins/Experimental/ControlFlows/Source/ControlFlows/Public/ControlFlowConcurrency.h b/Engine/Plugins/Experimental/ControlFlows/Source/ControlFlows/Public/ControlFlowConcurrency.h
--- a/Engine/Plugins/Experimental/ControlFlows/Source/ControlFlows/Public/ControlFlowConcurrency.h
+++ b/Engine/Plugins/Experimental/ControlFlows/Source/ControlFlows/Public/ControlFlowConcurrency.h
@@ -5,6 +5,7 @@
 #include "Delegates/Delegate.h"
 
 class FControlFlow;
+class FControlFlowSubTaskBase;
 class FConcurrentControlFlowBehavior;
 struct FConcurrencySubFlowContainer;
 
@@ -39,6 +40,8 @@ private:
 	bool bCancelAllHasBegun = false;
 	TMap<int32, TSharedRef<FConcurrencySubFlowContainer>> ConcurrentFlows;
 
+	TWeakPtr<FControlFlowSubTaskBase> OwningTask;
+
 private:
 	FConcurrentControlFlowBehavior GetConcurrencyBehavior() const;
 };
diff --git a/Engine/Plugins/Experimental/ControlFlows/Source/ControlFlows/Public/ControlFlowNode.h b/Engine/Plugins/Experimental/ControlFlows/Source/ControlFlows/Public/ControlFlowNode.h
--- a/Engine/Plugins/Experimental/ControlFlows/Source/ControlFlows/Public/ControlFlowNode.h
+++ b/Engine/Plugins/Experimental/ControlFlows/Source/ControlFlows/Public/ControlFlowNode.h
@@ -56,6 +56,7 @@ public:
 
 protected:
 	friend class FControlFlow;
+	friend class FControlFlowSubTaskBase;
 	
 	virtual void Execute() {}
 
diff --git a/Engine/Plugins/Experimental/ControlFlows/Source/ControlFlows/Public/ControlFlowTask.h b/Engine/Plugins/Experimental/ControlFlows/Source/ControlFlows/Public/ControlFlowTask.h
--- a/Engine/Plugins/Experimental/ControlFlows/Source/ControlFlows/Public/ControlFlowTask.h
+++ b/Engine/Plugins/Experimental/ControlFlows/Source/ControlFlows/Public/ControlFlowTask.h
@@ -14,6 +14,8 @@ public:
 	virtual ~FControlFlowSubTaskBase() {}
 
 	const FString& GetTaskName() const { return DebugName; }
+	void SetOwningNode(TSharedRef<FControlFlowNode_Task> InNode) { OwningNode = InNode; }
+	TWeakPtr<FControlFlow> GetOwningFlowForTaskNode() const;
 
 protected:
 	friend class FControlFlow;
@@ -24,6 +26,9 @@ protected:
 	virtual void Execute();
 	virtual void Cancel();
 
+protected:
+	TWeakPtr<FControlFlowNode_Task> OwningNode;
+
 private:
 	mutable FSimpleDelegate TaskCompleteCallback;
 	mutable FSimpleDelegate TaskCancelledCallback;
@@ -36,7 +41,7 @@ private:
 class FControlFlowSimpleSubTask : public FControlFlowSubTaskBase
 {
 public:
-	FControlFlowSimpleSubTask(const FString& TaskName, TSharedRef<FControlFlow> FlowOwner);
+	FControlFlowSimpleSubTask(const FString& TaskName);
 
 protected:
 	friend class FControlFlow;
@@ -63,7 +68,7 @@ class FControlFlowTask_LoopDeprecated : public FControlFlowSimpleSubTask
 {
 public:
 
-	FControlFlowTask_LoopDeprecated(FControlFlowLoopComplete& TaskCompleteDelegate, const FString& TaskName, TSharedRef<FControlFlow> FlowOwner);
+	FControlFlowTask_LoopDeprecated(FControlFlowLoopComplete& TaskCompleteDelegate, const FString& TaskName);
 
 protected:
 
@@ -121,7 +126,7 @@ private:
 
 	TMap<int32, TSharedRef<FControlFlow>> Branches;
 
-	int32 SelectedBranch = INDEX_NONE;
+	TOptional<int32> SelectedBranch;
 };
 
 //////////////////////////////////////////////////////////////////////////////
