Removing ensure and leaving comment that cancelling the next node in queue might be unnecesary

diff --git a/Engine/Plugins/Experimental/ControlFlows/Source/ControlFlows/Private/ControlFlow.cpp b/Engine/Plugins/Experimental/ControlFlows/Source/ControlFlows/Private/ControlFlow.cpp
--- a/Engine/Plugins/Experimental/ControlFlows/Source/ControlFlows/Private/ControlFlow.cpp
+++ b/Engine/Plugins/Experimental/ControlFlows/Source/ControlFlows/Private/ControlFlow.cpp
@@ -117,7 +117,7 @@ void FControlFlow::HandleControlFlowNodeCompleted(TSharedRef<const FControlFlowN
 						Activity->Update(*CurrentNode->GetNodeName());
 
 					FlowQueue.RemoveAt(0);
-					CurrentNode->CancelFlow();
+					CurrentNode->CancelFlow(); // This cancel might be completely unnecessary because the next node in Queue has not yet fired
 				}
 				else
 				{
diff --git a/Engine/Plugins/Experimental/ControlFlows/Source/ControlFlows/Private/ControlFlowTask_Loop.cpp b/Engine/Plugins/Experimental/ControlFlows/Source/ControlFlows/Private/ControlFlowTask_Loop.cpp
--- a/Engine/Plugins/Experimental/ControlFlows/Source/ControlFlows/Private/ControlFlowTask_Loop.cpp
+++ b/Engine/Plugins/Experimental/ControlFlows/Source/ControlFlows/Private/ControlFlowTask_Loop.cpp
@@ -53,7 +53,7 @@ void FControlFlowTask_ConditionalLoop::Execute()
 
 void FControlFlowTask_ConditionalLoop::Cancel()
 {
-	if (ensureAlways(ConditionalLoop.IsValid()) && ConditionalLoop->FlowLoop->IsRunning())
+	if (ConditionalLoop.IsValid() && ConditionalLoop->FlowLoop->IsRunning())
 	{
 		ConditionalLoop->FlowLoop->CancelFlow();
 	}
