Updated comment with explanation

diff --git a/Engine/Plugins/Experimental/ControlFlows/Source/ControlFlows/Private/ControlFlow.cpp b/Engine/Plugins/Experimental/ControlFlows/Source/ControlFlows/Private/ControlFlow.cpp
--- a/Engine/Plugins/Experimental/ControlFlows/Source/ControlFlows/Private/ControlFlow.cpp
+++ b/Engine/Plugins/Experimental/ControlFlows/Source/ControlFlows/Private/ControlFlow.cpp
@@ -117,7 +117,10 @@ void FControlFlow::HandleControlFlowNodeCompleted(TSharedRef<const FControlFlowN
 						Activity->Update(*CurrentNode->GetNodeName());
 
 					FlowQueue.RemoveAt(0);
-					CurrentNode->CancelFlow(); // This cancel might be completely unnecessary because the next node in Queue has not yet fired
+
+					/* Calling CancelFlow() on the next node will call CancelFlow() through the entire `FlowQueue` nodes because 
+					* Calling `FControlFlowNode::CancelFlow()` will call back into `FControlFlow::HandleControlFlowNodeCompleted()` and end here until `FlowQueue.Num() == 0` */
+					CurrentNode->CancelFlow(); 
 				}
 				else
 				{
