Implementing actual parallel execution as an option for FConcurrentControlFlows.
Fixed static memory in FControlFlowStatics with a static accessor

diff --git a/Engine/Plugins/Experimental/ControlFlows/Source/ControlFlows/Private/ControlFlowConcurrency.cpp b/Engine/Plugins/Experimental/ControlFlows/Source/ControlFlows/Private/ControlFlowConcurrency.cpp
--- a/Engine/Plugins/Experimental/ControlFlows/Source/ControlFlows/Private/ControlFlowConcurrency.cpp
+++ b/Engine/Plugins/Experimental/ControlFlows/Source/ControlFlows/Private/ControlFlowConcurrency.cpp
@@ -4,6 +4,7 @@
 #include "ControlFlow.h"
 #include "ControlFlows.h"
 #include "ControlFlowTask.h"
+#include "Async/ParallelFor.h"
 
 FConcurrencySubFlowContainer::FConcurrencySubFlowContainer(const FString& InDebugName)
 	: SubFlow(MakeShared<FControlFlow>(InDebugName))
@@ -100,6 +101,12 @@ FControlFlow& FConcurrentControlFlows::AddOrGetFlow(int32 InIdentifier, const FS
 	}
 }
 
+FConcurrentControlFlows& FConcurrentControlFlows::SetExecution(const EConcurrentExecution InBehavior)
+{
+	ExecutionBehavior = InBehavior;
+	return *this;
+}
+
 FControlFlow& FConcurrentControlFlows::AddOrGetProng(int32 InIdentifier, const FString& DebugSubFlowName /*= TEXT("")*/)
 {
 	return AddOrGetFlow(InIdentifier, DebugSubFlowName);
@@ -131,29 +138,32 @@ bool FConcurrentControlFlows::HasAnySubFlowBeenExecuted() const
 	return false;
 }
 
+bool FConcurrentControlFlows::HasAllSubFlowsBeenExecuted() const
+{
+	for (const TPair<int32, TSharedRef<FConcurrencySubFlowContainer>>& PairIt : ConcurrentFlows)
+	{
+		if (!PairIt.Value->HasBeenExecuted())
+		{
+			return false;
+		}
+	}
+
+	return true;
+}
+
 void FConcurrentControlFlows::HandleConcurrentFlowCompleted(int32 FlowIndex)
 {
 	check(ConcurrentFlows.Contains(FlowIndex));
 
 	UE_LOG(LogControlFlows, Verbose, TEXT("ConcurrentControlFlow Finished: %s"), *ConcurrentFlows[FlowIndex]->GetDebugName());
 
-	if (!bCancelAllHasBegun)
+	if (ExecutionBehavior == EConcurrentExecution::Parallel)
 	{
-		if (GetConcurrencyBehavior().GetContinueCondition() == FConcurrentControlFlowBehavior::EContinueConditions::ContinueOnAll_CompletedOrCancelled)
-		{
-			if (AreAllSubFlowsCompletedOrCancelled())
-			{
-				OnConcurrencyCompleted.ExecuteIfBound();
-			}
-			else
-			{
-				UE_LOG(LogControlFlows, Verbose, TEXT("Other flows are still running"));
-			}
-		}
-		else
-		{
-			checkf(false, TEXT("Unhandled Continue Condition"));
-		}
+		HandleParallelFlowExecutionFinished();
+	}
+	else if (!bCancelAllHasBegun)
+	{
+		CheckToBroadcastComplete();
 	}
 }
 
@@ -163,41 +173,89 @@ void FConcurrentControlFlows::HandleConcurrentFlowCancelled(int32 FlowIndex)
 
 	UE_LOG(LogControlFlows, Verbose, TEXT("ConcurrentControlFlow Cancelled %s"), *ConcurrentFlows[FlowIndex]->GetDebugName());
 
-	if (!bCancelAllHasBegun)
+	if (ExecutionBehavior == EConcurrentExecution::Parallel)
+	{
+		HandleParallelFlowExecutionFinished();
+	}
+	else if (!bCancelAllHasBegun)
 	{
 		HandleConcurrentFlowCompleted(FlowIndex);
 	}
 }
 
-void FConcurrentControlFlows::Execute()
+void FConcurrentControlFlows::HandleParallelFlowExecutionFinished()
 {
-	ensureAlwaysMsgf(!this->HasAnySubFlowBeenExecuted(), TEXT("Did you call ExecuteFlow() on a SubFlow? Do not do this! You only need to call ExecuteFlow once per FControlFlowStatics::Create!"));
-
-	bool bFlowExecuted = false;
-
-	for (TPair<int32, TSharedRef<FConcurrencySubFlowContainer>>& PairIt : ConcurrentFlows)
+	if (HasAllSubFlowsBeenExecuted())
 	{
-		UE_LOG(LogControlFlows, Verbose, TEXT("ConcurrentControlFlow::Execute - Executing Subflow %s"), *PairIt.Value->GetDebugName());
-
-		bFlowExecuted = true;
-
-		PairIt.Value->GetControlFlow()->LastZeroSecondDelay = OwningTask.Pin()->GetOwningFlowForTaskNode().Pin()->LastZeroSecondDelay;
-		PairIt.Value->Execute();
+		CheckToBroadcastComplete();
+	}
+}
 
-		if (bCancelAllHasBegun)
+void FConcurrentControlFlows::CheckToBroadcastComplete()
+{
+	if (GetConcurrencyBehavior().GetContinueCondition() == FConcurrentControlFlowBehavior::EContinueConditions::Default)
+	{
+		if (AreAllSubFlowsCompletedOrCancelled())
 		{
-			break;
+			OnConcurrencyCompleted.ExecuteIfBound();
+		}
+		else
+		{
+			UE_LOG(LogControlFlows, Verbose, TEXT("Other flows are still running"));
 		}
 	}
-	
-	if (ConcurrentFlows.Num() > 0)
+	else
 	{
-		ensureAlwaysMsgf(bFlowExecuted, TEXT("There were defined Subflows, but none were executed"));
+		checkf(false, TEXT("Unhandled Continue Condition"));
 	}
+}
 
-	if (!bFlowExecuted)
+void FConcurrentControlFlows::Execute()
+{
+	ensureAlwaysMsgf(!this->HasAnySubFlowBeenExecuted(), TEXT("Did you call ExecuteFlow() on a SubFlow? Do not do this! You only need to call ExecuteFlow once per FControlFlowStatics::Create!"));
+
+	if (ConcurrentFlows.Num() == 0)
 	{
 		OnAllCompleted();
+		return;
+	}
+
+	TArray<TSharedRef<FConcurrencySubFlowContainer>> FlowsToExecute;
+	for (const TPair<int32, TSharedRef<FConcurrencySubFlowContainer>>& PairIt : ConcurrentFlows)
+	{
+		FlowsToExecute.Add(PairIt.Value);
+		PairIt.Value->GetControlFlow()->LastZeroSecondDelay = OwningTask.Pin()->GetOwningFlowForTaskNode().Pin()->LastZeroSecondDelay;
+	}
+
+	if (ExecutionBehavior == EConcurrentExecution::Parallel)
+	{
+		ParallelFor(FlowsToExecute.Num(), [&FlowsToExecute](int32 Index)
+			{
+				UE_LOG(LogControlFlows, Verbose, TEXT("ConcurrentControlFlow::Execute - Executing Subflow %s"), *FlowsToExecute[Index]->GetDebugName());
+				FlowsToExecute[Index]->Execute();
+			});
+	}
+	else
+	{
+		if (ExecutionBehavior == EConcurrentExecution::Random)
+		{
+			for (int32 Idx = 0; Idx < FlowsToExecute.Num(); ++Idx)
+			{
+				Swap(FlowsToExecute[Idx], FlowsToExecute[FMath::RandRange(0, FlowsToExecute.Num() - 1)]);
+			}
+		}
+
+		for (const TSharedRef<FConcurrencySubFlowContainer>& Flow : FlowsToExecute)
+		{
+			UE_LOG(LogControlFlows, Verbose, TEXT("ConcurrentControlFlow::Execute - Executing Subflow %s"), *Flow->GetDebugName());
+
+			Flow->Execute();
+
+			if (bCancelAllHasBegun)
+			{
+				break;
+			}
+		}
 	}
 }
 
diff --git a/Engine/Plugins/Experimental/ControlFlows/Source/ControlFlows/Private/ControlFlowManager.cpp b/Engine/Plugins/Experimental/ControlFlows/Source/ControlFlows/Private/ControlFlowManager.cpp
--- a/Engine/Plugins/Experimental/ControlFlows/Source/ControlFlows/Private/ControlFlowManager.cpp
+++ b/Engine/Plugins/Experimental/ControlFlows/Source/ControlFlows/Private/ControlFlowManager.cpp
@@ -4,31 +4,24 @@
 #include "Containers/Ticker.h"
 #include "ControlFlows.h"
 
-static FTSTicker::FDelegateHandle NextFrameCheckForExecution;
-static FTSTicker::FDelegateHandle NextFrameCheckForFlowCleanup;
-
 TArray<TSharedRef<FControlFlowContainerBase>>& FControlFlowStatics::GetNewlyCreatedFlows()
 {
-	static TArray<TSharedRef<FControlFlowContainerBase>> NewlyCreatedFlows;
-	return NewlyCreatedFlows;
+	return Get().NewlyCreatedFlows;
 }
 
 TArray<TSharedRef<FControlFlowContainerBase>>& FControlFlowStatics::GetPersistentFlows()
 {
-	static TArray<TSharedRef<FControlFlowContainerBase>> PersistentFlows;
-	return PersistentFlows;
+	return Get().PersistentFlows;
 }
 
 TArray<TSharedRef<FControlFlowContainerBase>>& FControlFlowStatics::GetExecutingFlows()
 {
-	static TArray<TSharedRef<FControlFlowContainerBase>> ExecutingFlows;
-	return ExecutingFlows;
+	return Get().ExecutingFlows;
 }
 
 TArray<TSharedRef<FControlFlowContainerBase>>& FControlFlowStatics::GetFinishedFlows()
 {
-	static TArray<TSharedRef<FControlFlowContainerBase>> FinishedFlows;
-	return FinishedFlows;
+	return  Get().FinishedFlows;
 }
 
 void FControlFlowStatics::HandleControlFlowStartedNotification(TSharedRef<const FControlFlow> InFlow)
@@ -53,18 +46,23 @@ void FControlFlowStatics::HandleControlFlowStartedNotification(TSharedRef<const
 
 void FControlFlowStatics::CheckNewlyCreatedFlows()
 {
-	FTSTicker::GetCoreTicker().RemoveTicker(NextFrameCheckForExecution);
-	NextFrameCheckForExecution = FTSTicker::GetCoreTicker().AddTicker(FTickerDelegate::CreateStatic(&FControlFlowStatics::IterateThroughNewlyCreatedFlows));
+	if (!Get().NextFrameCheckForExecution.IsValid())
+	{
+		Get().NextFrameCheckForExecution = FTSTicker::GetCoreTicker().AddTicker(FTickerDelegate::CreateStatic(&FControlFlowStatics::IterateThroughNewlyCreatedFlows));
+	}
 }
 
 void FControlFlowStatics::CheckForInvalidFlows()
 {
-	FTSTicker::GetCoreTicker().RemoveTicker(NextFrameCheckForFlowCleanup);
-	NextFrameCheckForFlowCleanup = FTSTicker::GetCoreTicker().AddTicker(FTickerDelegate::CreateStatic(&FControlFlowStatics::IterateForInvalidFlows));
+	if (!Get().NextFrameCheckForFlowCleanup.IsValid())
+	{
+		Get().NextFrameCheckForFlowCleanup = FTSTicker::GetCoreTicker().AddTicker(FTickerDelegate::CreateStatic(&FControlFlowStatics::IterateForInvalidFlows));
+	}
 }
 
 bool FControlFlowStatics::IterateThroughNewlyCreatedFlows(float DeltaTime)
 {
+	Get().NextFrameCheckForExecution.Reset();
 	TArray<TSharedRef<FControlFlowContainerBase>>& NewFlows = GetNewlyCreatedFlows();
 	for (size_t Idx = 0; Idx < NewFlows.Num(); ++Idx)
 	{
@@ -98,6 +96,7 @@ bool FControlFlowStatics::IterateThroughNewlyCreatedFlows(float DeltaTime)
 
 bool FControlFlowStatics::IterateForInvalidFlows(float DeltaTime)
 {
+	Get().NextFrameCheckForFlowCleanup.Reset();
 	//Iterating through Persistent Flows
 	{
 		TArray<TSharedRef<FControlFlowContainerBase>>& Persistent = GetPersistentFlows();
@@ -167,3 +166,9 @@ bool FControlFlowStatics::IterateForInvalidFlows(float DeltaTime)
 
 	return false;
 }
+
+FControlFlowStatics& FControlFlowStatics::Get()
+{
+	static TUniquePtr<FControlFlowStatics> Singleton = MakeUnique<FControlFlowStatics>();
+	return *Singleton;
+}
diff --git a/Engine/Plugins/Experimental/ControlFlows/Source/ControlFlows/Public/ControlFlowConcurrency.h b/Engine/Plugins/Experimental/ControlFlows/Source/ControlFlows/Public/ControlFlowConcurrency.h
--- a/Engine/Plugins/Experimental/ControlFlows/Source/ControlFlows/Public/ControlFlowConcurrency.h
+++ b/Engine/Plugins/Experimental/ControlFlows/Source/ControlFlows/Public/ControlFlowConcurrency.h
@@ -9,13 +9,22 @@ class FControlFlowSubTaskBase;
 class FConcurrentControlFlowBehavior;
 struct FConcurrencySubFlowContainer;
 
-// All Flows will be executed concurrently (Single-Threaded). WARNING: Having a non-terminating loop within a fork can cause a forever hang for the forked step
+enum class EConcurrentExecution
+{
+	Default, // Flows will execute in a single thread; always in the same order
+	Random, // Flows will execute in a single thread; random order
+	Parallel, // Flows will execute in actual separate threads
+};
+
+// All Flows will be executed concurrently (or in-parallel, See: EConcurrentExecution). WARNING: Having a non-terminating loop within a fork can cause a forever hang for the forked step
 class FConcurrentControlFlows : public TSharedFromThis<FConcurrentControlFlows>
 {
 public:
 	CONTROLFLOWS_API FControlFlow& AddOrGetProng(int32 InIdentifier, const FString& DebugSubFlowName = TEXT(""));
 	CONTROLFLOWS_API FControlFlow& AddOrGetFlow(int32 InIdentifier, const FString& DebugSubFlowName = TEXT(""));
 
+	CONTROLFLOWS_API FConcurrentControlFlows& SetExecution(const EConcurrentExecution InBehavior);
+
 	// TODO: Multi-threaded-like behavior
 	// void Set(const FConcurrentControlFlowBehavior& InBehavior);
 
@@ -24,9 +33,13 @@ private:
 
 	bool AreAllSubFlowsCompletedOrCancelled() const;
 	bool HasAnySubFlowBeenExecuted() const;
+	bool HasAllSubFlowsBeenExecuted() const;
 
 	void HandleConcurrentFlowCompleted(int32 FlowIndex);
 	void HandleConcurrentFlowCancelled(int32 FlowIndex);
+	void HandleParallelFlowExecutionFinished();
+
+	void CheckToBroadcastComplete();
 
 	void Execute();
 	void CancelAll(); // Do not make public
@@ -37,6 +50,7 @@ private:
 	FSimpleDelegate OnConcurrencyCompleted;
 	FSimpleDelegate OnConcurrencyCancelled;
 
+	EConcurrentExecution ExecutionBehavior = EConcurrentExecution::Default;
 	bool bCancelAllHasBegun = false;
 	TMap<int32, TSharedRef<FConcurrencySubFlowContainer>> ConcurrentFlows;
 
@@ -72,30 +86,30 @@ private:
 	TSharedRef<FControlFlow> SubFlow;
 };
 
-// Placeholder class to extend Concurrency behavior. TODO: Read up more on multi-threaded processing
+// Placeholder class to extend Concurrency behavior
 class FConcurrentControlFlowBehavior
 {
 	friend class FConcurrentControlFlows;
 private:
 	enum class EContinueConditions
 	{
-		ContinueOnAll_CompletedOrCancelled, // Completed or Canceled
-		// Can add behavior here as needed before having an evaluation method
+		Default,
+		// TODO:
+		// Default: Equivalent to "Sync" in verse: Outer flow will continue once all flows have completed or been cancelled
+		//		See: "Default_Note"
+		// 
+		// Race: Outer flow will continue once the first flow completes or cancelled. Other flows cancelled
+		//		Race_Complete: Outer flow will continue once the first flow complete (ignoring cancel). Other flows cancelled
+		//		Race_Cancel: Outer flow will continue once the first flow cancel (ignoring complete). Other flows cancelled
+		// 
+		// Rush: Outer flow will continue once the first flow completes or cancelled. Other flows continue
+		//		Rush_Complete: Outer flow will continue once the first flow complete (ignoring cancel). Other flows continue
+		//		Rush_Cancel: Outer flow will continue once the first flow cancel (ignoring complete). Other flows continue
+		// 
+		// Default_Note: "Sync_Complete" and "Sync_Cancel" do not make sense and is unnecessary. If a flow is currently running, we have to wait for that flow to complete or cancel in order for the "Sync" condition to be satisfied.
+		// 
+		// Will allow to specify specific flows as an exception to Race/Rush
 	};
 
-	EContinueConditions GetContinueCondition() const { return EContinueConditions::ContinueOnAll_CompletedOrCancelled; }
-
-	// TODO: Customized continue conditions
-	// void EvaluateContinueCondition();
-	//
-	// Examples:
-	//		ContinueOnFirstCompleted_CancelOthers,
-	//		ContinueOnFirstCanceled_CancelOthers,
-	//		ContinueOnFirstCompletedOrCanceled_CancelOthers,
-	//		ContinueOnSpecificFlowsCompleted_CancelOthers,
-	//		ContinueOnSpecificFlowsCancelled_CancelOthers,
-	//		ContinueOnSpecificFlowsCompletedOrCancelled_CancelOthers
-
-	// TODO: Determine how a Forked Flow handles canceling
-	// Having a cancelled sub-flow affect other flows sort of goes against the idea of "Fork" flow, where each Subflow behaves independently from each other. Further debate required.
+	EContinueConditions GetContinueCondition() const { return EContinueConditions::Default; }
 };
diff --git a/Engine/Plugins/Experimental/ControlFlows/Source/ControlFlows/Public/ControlFlowManager.h b/Engine/Plugins/Experimental/ControlFlows/Source/ControlFlows/Public/ControlFlowManager.h
--- a/Engine/Plugins/Experimental/ControlFlows/Source/ControlFlows/Public/ControlFlowManager.h
+++ b/Engine/Plugins/Experimental/ControlFlows/Source/ControlFlows/Public/ControlFlowManager.h
@@ -5,6 +5,7 @@
 // IWYU pragma: begin_keep
 #include "ControlFlow.h"
 #include "UObject/StrongObjectPtr.h"
+#include "Containers/Ticker.h"
 #include "ControlFlowContainer.h"
 // IWYU pragma: end_keep
 
@@ -146,4 +147,16 @@ private:
 private:
 	static bool IterateThroughNewlyCreatedFlows(float DeltaTime);
 	static bool IterateForInvalidFlows(float DeltaTime);
+
+private:
+	static FControlFlowStatics& Get();
+
+private:
+	TArray<TSharedRef<FControlFlowContainerBase>> NewlyCreatedFlows;
+	TArray<TSharedRef<FControlFlowContainerBase>> PersistentFlows;
+	TArray<TSharedRef<FControlFlowContainerBase>> ExecutingFlows;
+	TArray<TSharedRef<FControlFlowContainerBase>> FinishedFlows;
+
+	FTSTicker::FDelegateHandle NextFrameCheckForExecution;
+	FTSTicker::FDelegateHandle NextFrameCheckForFlowCleanup;
 };
