Fixing bug where during OnStepCompletedDelegate.Broadcast(), IsRunning() returns false

diff --git a/Engine/Plugins/Experimental/ControlFlows/Source/ControlFlows/Private/ControlFlow.cpp b/Engine/Plugins/Experimental/ControlFlows/Source/ControlFlows/Private/ControlFlow.cpp
--- a/Engine/Plugins/Experimental/ControlFlows/Source/ControlFlows/Private/ControlFlow.cpp
+++ b/Engine/Plugins/Experimental/ControlFlows/Source/ControlFlows/Private/ControlFlow.cpp
@@ -124,7 +124,9 @@ void FControlFlow::HandleControlFlowNodeCompleted(TSharedRef<const FControlFlowN
 				}
 				else
 				{
+					bBroadcastingStepComplete = true;
 					OnStepCompletedDelegate.Broadcast();
+					bBroadcastingStepComplete = false;
 					ExecuteNextNodeInQueue();
 				}
 			}
@@ -254,6 +256,11 @@ void FControlFlow::Reset()
 	UnnamedBranchCounter = 0;
 }
 
+bool FControlFlow::IsRunning() const
+{
+	return CurrentNode.IsValid() || bBroadcastingStepComplete;
+}
+
 void FControlFlow::CancelFlow()
 {
 	if (CurrentNode.IsValid())
diff --git a/Engine/Plugins/Experimental/ControlFlows/Source/ControlFlows/Public/ControlFlow.h b/Engine/Plugins/Experimental/ControlFlows/Source/ControlFlows/Public/ControlFlow.h
--- a/Engine/Plugins/Experimental/ControlFlows/Source/ControlFlows/Public/ControlFlow.h
+++ b/Engine/Plugins/Experimental/ControlFlows/Source/ControlFlows/Public/ControlFlow.h
@@ -131,7 +131,7 @@ public:
 	/** This needs to be called, otherwise nothing will happen!! Call after you finish adding functions to the queue. Calling with an empty queue is safe. */
 	CONTROLFLOWS_API void ExecuteFlow();
 	CONTROLFLOWS_API void Reset();
-	bool IsRunning() const { return CurrentNode.IsValid(); }
+	CONTROLFLOWS_API bool IsRunning() const;
 	size_t NumInQueue() const { return FlowQueue.Num(); }
 
 	FSimpleMulticastDelegate& OnNodeComplete() const { return OnStepCompletedDelegate; }
@@ -384,6 +384,8 @@ private:
 
 	bool bInterpretCancelledNodeAsComplete = false;
 
+	bool bBroadcastingStepComplete = false;
+
 	int32 UnnamedNodeCounter = 0;
 
 	int32 UnnamedBranchCounter = 0;
