Deleting `CastedTo.h` and `SmartPtr.h` files
Replacing `TWeakContainer` and with `TWeakPtrVariant`
Signficantly reducing the number of times `OwningObject` gets pinned.

diff --git a/Engine/Plugins/Experimental/ControlFlows/Source/ControlFlows/Private/ControlFlowManager.cpp b/Engine/Plugins/Experimental/ControlFlows/Source/ControlFlows/Private/ControlFlowManager.cpp
--- a/Engine/Plugins/Experimental/ControlFlows/Source/ControlFlows/Private/ControlFlowManager.cpp
+++ b/Engine/Plugins/Experimental/ControlFlows/Source/ControlFlows/Private/ControlFlowManager.cpp
@@ -29,7 +29,7 @@ void FControlFlowStatics::HandleControlFlowStartedNotification(TSharedRef<const
 	TArray<TSharedRef<FControlFlowContainerBase>>& NewFlows = GetNewlyCreatedFlows();
 	for (size_t Idx = 0; Idx < NewFlows.Num(); ++Idx)
 	{
-		if (ensureAlways(NewFlows[Idx]->OwningObjectIsValid()))
+		if (ensureAlways(UE::Private::OwningObjectIsValid(NewFlows[Idx])))
 		{
 			if (InFlow == NewFlows[Idx]->GetControlFlow())
 			{
@@ -66,7 +66,7 @@ bool FControlFlowStatics::IterateThroughNewlyCreatedFlows(float DeltaTime)
 	TArray<TSharedRef<FControlFlowContainerBase>>& NewFlows = GetNewlyCreatedFlows();
 	for (size_t Idx = 0; Idx < NewFlows.Num(); ++Idx)
 	{
-		if (ensureAlways(NewFlows[Idx]->OwningObjectIsValid()))
+		if (ensureAlways(UE::Private::OwningObjectIsValid(NewFlows[Idx])))
 		{
 			TSharedRef<FControlFlow> NewFlow = NewFlows[Idx]->GetControlFlow();
 			if (ensureAlwaysMsgf(NewFlow->IsRunning(), TEXT("Call to execute after queue-ing your steps to avoid this ensure. We will fire the flow 1 frame late to hopefully not cause anything from breaking. Flow:%s"), *NewFlow->GetDebugName()))
@@ -102,7 +102,7 @@ bool FControlFlowStatics::IterateForInvalidFlows(float DeltaTime)
 		TArray<TSharedRef<FControlFlowContainerBase>>& Persistent = GetPersistentFlows();
 		for (size_t Idx = 0; Idx < Persistent.Num(); ++Idx)
 		{
-			if (Persistent[Idx]->OwningObjectIsValid())
+			if (UE::Private::OwningObjectIsValid(Persistent[Idx]))
 			{
 				TSharedRef<FControlFlow> PersistentFlow = Persistent[Idx]->GetControlFlow();
 				if (PersistentFlow->IsRunning())
@@ -125,12 +125,12 @@ bool FControlFlowStatics::IterateForInvalidFlows(float DeltaTime)
 		TArray<TSharedRef<FControlFlowContainerBase>>& Executing = GetExecutingFlows();
 		for (size_t Idx = 0; Idx < Executing.Num(); ++Idx)
 		{
-			if (Executing[Idx]->OwningObjectIsValid())
+			if (UE::Private::OwningObjectIsValid(Executing[Idx]))
 			{
 				TSharedRef<FControlFlow> ExecutingFlow = Executing[Idx]->GetControlFlow();
 				if (!ExecutingFlow->IsRunning() && ensureAlways(ExecutingFlow->NumInQueue() == 0))
 				{
-					Executing[Idx]->ControlFlow->Activity = nullptr;
+					Executing[Idx]->GetControlFlow()->Activity = nullptr;
 					GetFinishedFlows().Add(Executing[Idx]);
 					Executing.RemoveAtSwap(Idx);
 					--Idx;
@@ -138,7 +138,7 @@ bool FControlFlowStatics::IterateForInvalidFlows(float DeltaTime)
 			}
 			else
 			{
-				Executing[Idx]->ControlFlow->Activity = nullptr;
+				Executing[Idx]->GetControlFlow()->Activity = nullptr;
 				Executing.RemoveAtSwap(Idx);
 				--Idx;
 			}
@@ -150,7 +150,7 @@ bool FControlFlowStatics::IterateForInvalidFlows(float DeltaTime)
 		TArray<TSharedRef<FControlFlowContainerBase>>& Completed = GetFinishedFlows();
 		for (size_t Idx = 0; Idx < Completed.Num(); ++Idx)
 		{
-			if (!Completed[Idx]->OwningObjectIsValid())
+			if (!UE::Private::OwningObjectIsValid(Completed[Idx]))
 			{
 				UE_LOG(LogControlFlows, Warning, TEXT("Owning Object for completed flow is not valid!"));
 			}
diff --git a/Engine/Plugins/Experimental/ControlFlows/Source/ControlFlows/Public/CastedTo.h b/Engine/Plugins/Experimental/ControlFlows/Source/ControlFlows/Public/CastedTo.h
deleted file mode 100644
diff --git a/Engine/Plugins/Experimental/ControlFlows/Source/ControlFlows/Public/ControlFlowContainer.h b/Engine/Plugins/Experimental/ControlFlows/Source/ControlFlows/Public/ControlFlowContainer.h
--- a/Engine/Plugins/Experimental/ControlFlows/Source/ControlFlows/Public/ControlFlowContainer.h
+++ b/Engine/Plugins/Experimental/ControlFlows/Source/ControlFlows/Public/ControlFlowContainer.h
@@ -6,13 +6,20 @@
 
 #include "Templates/SharedPointer.h"
 #include "Templates/PointerVariants.h"
-#include "CastedTo.h"
+
+class FControlFlowContainerBase;
+namespace UE::Private
+{
+	static bool OwningObjectIsValid(TSharedRef<const FControlFlowContainerBase> InFlowContainer);
+}
+
+/* These wrapper classes along with FControlFlowStatics should be restructured in a way so that we do not rely on heap memory anymore
+*  Specifically we do not derive any FControlFlow classes from `TSharedFromThis` nor `UObject`
+*/
 
 class FControlFlowContainerBase : public TSharedFromThis<FControlFlowContainerBase>
 {
-	friend class FControlFlowStatics;
 public:
-
 	FControlFlowContainerBase() = delete;
 	virtual ~FControlFlowContainerBase() {}
 
@@ -22,16 +29,17 @@ public:
 		checkf(FlowName.Len() > 0, TEXT("All Flows need a non-empty ID!"));
 	}
 
-private:
-	virtual const void* const GetOwningObject() const = 0;
-
-private:
-	bool OwningObjectIsValid() const { return GetOwningObject() != nullptr; }
+public:
+	virtual bool OwningObjectEqualTo(const void* InObject) const = 0;
 
 	const FString& GetFlowName() const { return FlowName; }
 
 	TSharedRef<FControlFlow> GetControlFlow() const { return ControlFlow; }
 
+private:
+	virtual bool OwningObjectIsValid() const = 0;
+	friend bool UE::Private::OwningObjectIsValid(TSharedRef<const FControlFlowContainerBase> InFlowContainer);
+
 private:
 	TSharedRef<FControlFlow> ControlFlow;
 
@@ -49,11 +57,24 @@ public:
 	{}
 
 private:
-	virtual const void* const GetOwningObject() const override final
+	virtual bool OwningObjectEqualTo(const void* InObject) const override final
 	{
-		return OwningObject.IsValid() ? OwningObject.Cast() : nullptr;
+		return InObject != nullptr && OwningObject.IsValid() && OwningObject.Pin().Get() == InObject;
+	}
+
+	virtual bool OwningObjectIsValid() const override final
+	{
+		return OwningObject.IsValid();
 	}
 
 private:
-	TWeakContainer<T> OwningObject;
-};
\ No newline at end of file
+	TWeakPtrVariant<T> OwningObject;
+};
+
+namespace UE::Private
+{
+	static inline bool OwningObjectIsValid(TSharedRef<const FControlFlowContainerBase> InFlowContainer)
+	{
+		return InFlowContainer->OwningObjectIsValid();
+	}
+}
\ No newline at end of file
diff --git a/Engine/Plugins/Experimental/ControlFlows/Source/ControlFlows/Public/ControlFlowManager.h b/Engine/Plugins/Experimental/ControlFlows/Source/ControlFlows/Public/ControlFlowManager.h
--- a/Engine/Plugins/Experimental/ControlFlows/Source/ControlFlows/Public/ControlFlowManager.h
+++ b/Engine/Plugins/Experimental/ControlFlows/Source/ControlFlows/Public/ControlFlowManager.h
@@ -117,9 +117,9 @@ private:
 		{
 			for (TSharedRef<FControlFlowContainerBase> ControlFlowContainer : InArray)
 			{
-				if (ControlFlowContainer->OwningObjectIsValid())
+				if (UE::Private::OwningObjectIsValid(ControlFlowContainer))
 				{
-					if (ControlFlowContainer->GetOwningObject() == OwningObject && ControlFlowContainer->GetFlowName() == FlowId)
+					if (ControlFlowContainer->GetFlowName() == FlowId && ControlFlowContainer->OwningObjectEqualTo(OwningObject))
 					{
 						return ControlFlowContainer->GetControlFlow();
 					}
diff --git a/Engine/Plugins/Experimental/ControlFlows/Source/ControlFlows/Public/SmartPtr.h b/Engine/Plugins/Experimental/ControlFlows/Source/ControlFlows/Public/SmartPtr.h
deleted file mode 100644
