Adding forward declare

diff --git a/Engine/Plugins/Experimental/ControlFlows/Source/ControlFlows/Public/ControlFlowContainer.h b/Engine/Plugins/Experimental/ControlFlows/Source/ControlFlows/Public/ControlFlowContainer.h
--- a/Engine/Plugins/Experimental/ControlFlows/Source/ControlFlows/Public/ControlFlowContainer.h
+++ b/Engine/Plugins/Experimental/ControlFlows/Source/ControlFlows/Public/ControlFlowContainer.h
@@ -8,6 +8,8 @@
 #include "Templates/PointerVariants.h"
 
 class FControlFlowContainerBase;
+class FControlFlow;
+
 namespace UE::Private
 {
 	static bool OwningObjectIsValid(TSharedRef<const FControlFlowContainerBase> InFlowContainer);
