Adding SplitFlow and SwitchFlow aliases

diff --git a/Engine/Plugins/Experimental/ControlFlows/Source/ControlFlows/Public/ControlFlow.h b/Engine/Plugins/Experimental/ControlFlows/Source/ControlFlows/Public/ControlFlow.h
--- a/Engine/Plugins/Experimental/ControlFlows/Source/ControlFlows/Public/ControlFlow.h
+++ b/Engine/Plugins/Experimental/ControlFlows/Source/ControlFlows/Public/ControlFlow.h
@@ -187,6 +187,12 @@ public:
 		return *this;
 	}
 
+	template<typename...ArgsT>
+	FControlFlow& SwitchFlow(ArgsT...Args)
+	{
+		return BranchFlow(Args...);
+	}
+
 	template<typename FunctionT, typename...ArgsT>
 	FControlFlow& BranchFlow(FunctionT InBranchLambda, ArgsT...Params)
 	{
@@ -194,6 +200,12 @@ public:
 		return *this;
 	}
 
+	template<typename...ArgsT>
+	FControlFlow& SplitFlow(ArgsT...Args)
+	{
+		return ForkFlow(Args...);
+	}
+
 	template<typename FunctionT, typename...ArgsT>
 	FControlFlow& ForkFlow(FunctionT InForkLambda, ArgsT...Params)
 	{
