Replacing forward declare with include

diff --git a/Engine/Plugins/Experimental/ControlFlows/Source/ControlFlows/Public/ControlFlow.h b/Engine/Plugins/Experimental/ControlFlows/Source/ControlFlows/Public/ControlFlow.h
--- a/Engine/Plugins/Experimental/ControlFlows/Source/ControlFlows/Public/ControlFlow.h
+++ b/Engine/Plugins/Experimental/ControlFlows/Source/ControlFlows/Public/ControlFlow.h
@@ -3,8 +3,7 @@
 #pragma once
 
 #include "ControlFlowNode.h"
-
-template <typename OptionalType> struct TOptional;
+#include "Misc/OptionalFwd.h"
 
 class UObject;
 class FControlFlowTask_Branch;
