Fixing assert/bug with Parallel Control Flow Execution

diff --git a/Engine/Plugins/Experimental/ControlFlows/Source/ControlFlows/Private/ControlFlowConcurrency.cpp b/Engine/Plugins/Experimental/ControlFlows/Source/ControlFlows/Private/ControlFlowConcurrency.cpp
--- a/Engine/Plugins/Experimental/ControlFlows/Source/ControlFlows/Private/ControlFlowConcurrency.cpp
+++ b/Engine/Plugins/Experimental/ControlFlows/Source/ControlFlows/Private/ControlFlowConcurrency.cpp
@@ -1,6 +1,8 @@
 // Copyright Epic Games, Inc. All Rights Reserved.
 
 #include "ControlFlowConcurrency.h"
+
+#include "Containers/Ticker.h"
 #include "ControlFlow.h"
 #include "ControlFlows.h"
 #include "ControlFlowTask.h"
@@ -26,12 +28,13 @@ bool FConcurrencySubFlowContainer::IsCompleteOrCancelled() const
 	return bIsComplete || HasBeenCancelled();
 }
 
-void FConcurrencySubFlowContainer::Execute()
+void FConcurrencySubFlowContainer::Execute(TWeakPtr<FConcurrentControlFlows::FParallelFlowLock> InFlowLock)
 {
 	UE_LOG(LogControlFlows, Verbose, TEXT("ConcurrentControlFlowContainer - Executing %s"), *GetDebugName());
 
 	if (ensureAlways(!IsCompleteOrCancelled()))
 	{
+		ParallelFlowLock = InFlowLock.IsValid() ? InFlowLock.Pin() : nullptr;
 		bHasBeenExecuted = true;
 		SubFlow->ExecuteFlow();
 	}
@@ -92,9 +95,9 @@ FControlFlow& FConcurrentControlFlows::AddOrGetFlow(int32 InIdentifier, const FS
 
 		NewConcurrentFlow->GetControlFlow()->ParentFlow = OwningTask.Pin()->GetOwningFlowForTaskNode();
 
-		NewConcurrentFlow->OnComplete().BindSP(SharedThis(this), &FConcurrentControlFlows::HandleConcurrentFlowCompleted, InIdentifier);
-		NewConcurrentFlow->OnExecutedWithoutAnyNodes().BindSP(SharedThis(this), &FConcurrentControlFlows::HandleConcurrentFlowCompleted, InIdentifier);
-		NewConcurrentFlow->OnCancelled().BindSP(SharedThis(this), &FConcurrentControlFlows::HandleConcurrentFlowCancelled, InIdentifier);
+		NewConcurrentFlow->OnComplete().BindSP(SharedThis(this), &FConcurrentControlFlows::HandleConcurrentFlowDone, InIdentifier, true);
+		NewConcurrentFlow->OnExecutedWithoutAnyNodes().BindSP(SharedThis(this), &FConcurrentControlFlows::HandleConcurrentFlowDone, InIdentifier, true);
+		NewConcurrentFlow->OnCancelled().BindSP(SharedThis(this), &FConcurrentControlFlows::HandleConcurrentFlowDone, InIdentifier, false);
 
 		ConcurrentFlows.Add(InIdentifier, NewConcurrentFlow);
 		return NewConcurrentFlow->GetControlFlow().Get();
@@ -138,28 +141,15 @@ bool FConcurrentControlFlows::HasAnySubFlowBeenExecuted() const
 	return false;
 }
 
-bool FConcurrentControlFlows::HasAllSubFlowsBeenExecuted() const
-{
-	for (const TPair<int32, TSharedRef<FConcurrencySubFlowContainer>>& PairIt : ConcurrentFlows)
-	{
-		if (!PairIt.Value->HasBeenExecuted())
-		{
-			return false;
-		}
-	}
-
-	return true;
-}
-
-void FConcurrentControlFlows::HandleConcurrentFlowCompleted(int32 FlowIndex)
+void FConcurrentControlFlows::HandleConcurrentFlowDone(int32 FlowIndex, const bool bCompleted)
 {
 	check(ConcurrentFlows.Contains(FlowIndex));
 
-	UE_LOG(LogControlFlows, Verbose, TEXT("ConcurrentControlFlow Finished: %s"), *ConcurrentFlows[FlowIndex]->GetDebugName());
+	UE_LOG(LogControlFlows, Verbose, TEXT("ConcurrentControlFlow %s: %s"), bCompleted ? TEXT("Finished") : TEXT("Cancelled"), *ConcurrentFlows[FlowIndex]->GetDebugName());
 
 	if (ExecutionBehavior == EConcurrentExecution::Parallel)
 	{
-		HandleParallelFlowExecutionFinished();
+		ConcurrentFlows[FlowIndex].Get().ParallelFlowLock.Reset();
 	}
 	else if (!bCancelAllHasBegun)
 	{
@@ -167,30 +157,6 @@ void FConcurrentControlFlows::HandleConcurrentFlowCompleted(int32 FlowIndex)
 	}
 }
 
-void FConcurrentControlFlows::HandleConcurrentFlowCancelled(int32 FlowIndex)
-{
-	check(ConcurrentFlows.Contains(FlowIndex));
-
-	UE_LOG(LogControlFlows, Verbose, TEXT("ConcurrentControlFlow Cancelled %s"), *ConcurrentFlows[FlowIndex]->GetDebugName());
-
-	if (ExecutionBehavior == EConcurrentExecution::Parallel)
-	{
-		HandleParallelFlowExecutionFinished();
-	}
-	else if (!bCancelAllHasBegun)
-	{
-		HandleConcurrentFlowCompleted(FlowIndex);
-	}
-}
-
-void FConcurrentControlFlows::HandleParallelFlowExecutionFinished()
-{
-	if (HasAllSubFlowsBeenExecuted())
-	{
-		CheckToBroadcastComplete();
-	}
-}
-
 void FConcurrentControlFlows::CheckToBroadcastComplete()
 {
 	if (GetConcurrencyBehavior().GetContinueCondition() == FConcurrentControlFlowBehavior::EContinueConditions::Default)
@@ -229,11 +195,35 @@ void FConcurrentControlFlows::Execute()
 
 	if (ExecutionBehavior == EConcurrentExecution::Parallel)
 	{
-		ParallelFor(FlowsToExecute.Num(), [&FlowsToExecute](int32 Index)
+		ParallelFlowLock = MakeShared<FConcurrentControlFlows::FParallelFlowLock, ESPMode::ThreadSafe>();
+
+		ParallelFor(FlowsToExecute.Num(), [this, &FlowsToExecute](int32 Index)
 			{
 				UE_LOG(LogControlFlows, Verbose, TEXT("ConcurrentControlFlow::Execute - Executing Subflow %s"), *FlowsToExecute[Index]->GetDebugName());
-				FlowsToExecute[Index]->Execute();
+				FlowsToExecute[Index]->Execute(ParallelFlowLock);
 			});
+
+		if (ParallelFlowLock.IsUnique())
+		{
+			ensure(AreAllSubFlowsCompletedOrCancelled());
+			CheckToBroadcastComplete();
+		}
+		else
+		{
+			FTSTicker::GetCoreTicker().AddTicker(FTickerDelegate::CreateSPLambda(this, [this](float)
+				{
+					if (ParallelFlowLock.IsUnique())
+					{
+						ensure(AreAllSubFlowsCompletedOrCancelled());
+						CheckToBroadcastComplete();
+						return false;
+					}
+					else
+					{
+						return true;
+					}
+				}), 0.0f);
+		}
 	}
 	else
 	{
@@ -249,7 +239,7 @@ void FConcurrentControlFlows::Execute()
 		{
 			UE_LOG(LogControlFlows, Verbose, TEXT("ConcurrentControlFlow::Execute - Executing Subflow %s"), *Flow->GetDebugName());
 
-			Flow->Execute();
+			Flow->Execute(nullptr);
 
 			if (bCancelAllHasBegun)
 			{
diff --git a/Engine/Plugins/Experimental/ControlFlows/Source/ControlFlows/Public/ControlFlowConcurrency.h b/Engine/Plugins/Experimental/ControlFlows/Source/ControlFlows/Public/ControlFlowConcurrency.h
--- a/Engine/Plugins/Experimental/ControlFlows/Source/ControlFlows/Public/ControlFlowConcurrency.h
+++ b/Engine/Plugins/Experimental/ControlFlows/Source/ControlFlows/Public/ControlFlowConcurrency.h
@@ -20,6 +20,8 @@ enum class EConcurrentExecution
 class FConcurrentControlFlows : public TSharedFromThis<FConcurrentControlFlows>
 {
 public:
+	struct FParallelFlowLock {};
+
 	CONTROLFLOWS_API FControlFlow& AddOrGetProng(int32 InIdentifier, const FString& DebugSubFlowName = TEXT(""));
 	CONTROLFLOWS_API FControlFlow& AddOrGetFlow(int32 InIdentifier, const FString& DebugSubFlowName = TEXT(""));
 
@@ -33,11 +35,8 @@ private:
 
 	bool AreAllSubFlowsCompletedOrCancelled() const;
 	bool HasAnySubFlowBeenExecuted() const;
-	bool HasAllSubFlowsBeenExecuted() const;
 
-	void HandleConcurrentFlowCompleted(int32 FlowIndex);
-	void HandleConcurrentFlowCancelled(int32 FlowIndex);
-	void HandleParallelFlowExecutionFinished();
+	void HandleConcurrentFlowDone(int32 FlowIndex, const bool bCompleted);
 
 	void CheckToBroadcastComplete();
 
@@ -52,6 +51,7 @@ private:
 
 	EConcurrentExecution ExecutionBehavior = EConcurrentExecution::Default;
 	bool bCancelAllHasBegun = false;
+	TSharedPtr<FConcurrentControlFlows::FParallelFlowLock, ESPMode::ThreadSafe> ParallelFlowLock;
 	TMap<int32, TSharedRef<FConcurrencySubFlowContainer>> ConcurrentFlows;
 
 	TWeakPtr<FControlFlowSubTaskBase> OwningTask;
@@ -71,7 +71,7 @@ private:
 	bool HasBeenCancelled() const;
 	bool IsCompleteOrCancelled() const;
 
-	void Execute();
+	void Execute(TWeakPtr<FConcurrentControlFlows::FParallelFlowLock> InFlowLock);
 	void Cancel();
 	FSimpleDelegate& OnComplete();
 	FSimpleDelegate& OnExecutedWithoutAnyNodes();
@@ -84,6 +84,7 @@ private:
 	bool bHasBeenCancelled = false;
 
 	TSharedRef<FControlFlow> SubFlow;
+	TSharedPtr<FConcurrentControlFlows::FParallelFlowLock, ESPMode::ThreadSafe> ParallelFlowLock;
 };
 
 // Placeholder class to extend Concurrency behavior
