[ControlFlows] Cache Parent.Pin() early to shorten call stacks, eliminate weak-ptr race window & extra atomics. Switch HandleControlFlowNodeCompleted to const TSharedRef& to avoid shared-ptr construction overhead,

diff --git a/Engine/Plugins/Experimental/ControlFlows/Source/ControlFlows/Private/ControlFlow.cpp b/Engine/Plugins/Experimental/ControlFlows/Source/ControlFlows/Private/ControlFlow.cpp
--- a/Engine/Plugins/Experimental/ControlFlows/Source/ControlFlows/Private/ControlFlow.cpp
+++ b/Engine/Plugins/Experimental/ControlFlows/Source/ControlFlows/Private/ControlFlow.cpp
@@ -71,7 +71,7 @@ void FControlFlow::BroadcastCancellation()
 	OnCancelledDelegate_Internal.ExecuteIfBound();
 }
 
-void FControlFlow::HandleControlFlowNodeCompleted(TSharedRef<const FControlFlowNode> NodeCompleted)
+void FControlFlow::HandleControlFlowNodeCompleted(const TSharedRef<const FControlFlowNode>& NodeCompleted)
 {
 	UE_LOG(LogControlFlows, Verbose, TEXT("ControlFlow - Executing %s%s(x%d).%s (FlowControlNodeCompleted)"), *GetFlowPath(), *DebugName, GetRepeatedFlowCount(), *NodeCompleted->GetNodeName());
 
diff --git a/Engine/Plugins/Experimental/ControlFlows/Source/ControlFlows/Private/ControlFlowNode.cpp b/Engine/Plugins/Experimental/ControlFlows/Source/ControlFlows/Private/ControlFlowNode.cpp
--- a/Engine/Plugins/Experimental/ControlFlows/Source/ControlFlows/Private/ControlFlowNode.cpp
+++ b/Engine/Plugins/Experimental/ControlFlows/Source/ControlFlows/Private/ControlFlowNode.cpp
@@ -22,35 +22,39 @@ FControlFlowNode::~FControlFlowNode()
 
 void FControlFlowNode::ContinueFlow()
 {
-	if (ensure(Parent.IsValid()))
+	TSharedPtr<FControlFlow> PinnedParent = Parent.Pin();
+	if (ensure(PinnedParent))
 	{
-		Parent.Pin()->HandleControlFlowNodeCompleted(SharedThis(this));
+		PinnedParent->HandleControlFlowNodeCompleted(SharedThis(this));
 	}
 }
 
 void FControlFlowNode::CancelFlow()
 {
-	if (ensure(Parent.IsValid()))
+	TSharedPtr<FControlFlow> PinnedParent = Parent.Pin();
+	if (ensure(PinnedParent))
 	{
 		bCancelled = true;
-		Parent.Pin()->HandleControlFlowNodeCompleted(SharedThis(this));
+		PinnedParent->HandleControlFlowNodeCompleted(SharedThis(this));
 	}
 }
 
 TSharedPtr<FTrackedActivity> FControlFlowNode::GetTrackedActivity() const
 {
-	if (ensure(Parent.IsValid()))
+	TSharedPtr<FControlFlow> PinnedParent = Parent.Pin();
+	if (ensure(PinnedParent))
 	{
-		return Parent.Pin()->GetTrackedActivity();
+		return PinnedParent->GetTrackedActivity();
 	}
 	return nullptr;
 }
 
 void FControlFlowNode::LogExecution()
 {
-	if (ensure(Parent.IsValid()))
+	TSharedPtr<FControlFlow> PinnedParent = Parent.Pin();
+	if (ensure(PinnedParent))
 	{
-		Parent.Pin()->LogNodeExecution(*this);
+		PinnedParent->LogNodeExecution(*this);
 	}
 }
 
@@ -61,9 +65,10 @@ FString FControlFlowNode::GetNodeName() const
 
 void FControlFlowNode::SetProfilerEventStarted()
 {
-	if (ensure(Parent.IsValid()))
+	TSharedPtr<FControlFlow> PinnedParent = Parent.Pin();
+	if (ensure(PinnedParent))
 	{
-		Parent.Pin()->SetProfilerEventStarted();
+		PinnedParent->SetProfilerEventStarted();
 	}
 }
 
@@ -85,9 +90,10 @@ void FControlFlowNode_SelfCompleting::Execute()
 {
 	LogExecution();
 
-	if (ensureAlways(Parent.IsValid()))
+	TSharedPtr<FControlFlow> PinnedParent = Parent.Pin();
+	if (ensureAlways(PinnedParent))
 	{
-		Parent.Pin()->ExecuteNode(SharedThis(this));
+		PinnedParent->ExecuteNode(SharedThis(this));
 	}
 }
 
@@ -110,9 +116,10 @@ void FControlFlowNode_RequiresCallback::Execute()
 
 	if (Process.IsBound())
 	{
-		if (ensureAlways(Parent.IsValid()))
+		TSharedPtr<FControlFlow> PinnedParent = Parent.Pin();
+		if (ensureAlways(PinnedParent))
 		{
-			UE_LOG(LogControlFlows, Verbose, TEXT("ControlFlow - Executing %s.%s -  Waiting On Callback"), *(Parent.Pin()->GetFlowPath().Append(Parent.Pin()->GetDebugName())), *NodeName);
+			UE_LOG(LogControlFlows, Verbose, TEXT("ControlFlow - Executing %s.%s -  Waiting On Callback"), *(PinnedParent->GetFlowPath().Append(PinnedParent->GetDebugName())), *NodeName);
 		}
 
 		Process.Execute(SharedThis(this));
diff --git a/Engine/Plugins/Experimental/ControlFlows/Source/ControlFlows/Public/ControlFlow.h b/Engine/Plugins/Experimental/ControlFlows/Source/ControlFlows/Public/ControlFlow.h
--- a/Engine/Plugins/Experimental/ControlFlows/Source/ControlFlows/Public/ControlFlow.h
+++ b/Engine/Plugins/Experimental/ControlFlows/Source/ControlFlows/Public/ControlFlow.h
@@ -350,7 +350,7 @@ public:
 private:
 	inline void SetProfilerEventStarted() { ensureMsgf(!bProfilerEventStarted, TEXT("Started a new control flow profiler event before previous step completed.")); bProfilerEventStarted = true; }
 
-	void HandleControlFlowNodeCompleted(TSharedRef<const FControlFlowNode> NodeCompleted);
+	void HandleControlFlowNodeCompleted(const TSharedRef<const FControlFlowNode>& NodeCompleted);
 
 	mutable FSimpleMulticastDelegate OnStepCompletedDelegate;
 	mutable FSimpleMulticastDelegate OnFlowCompleteDelegate;
