[ControlFlow] Make sure current node is valid during step complete notification

#tests Used callback to track which step was active for debugging

#Synced-CL 50458449

diff --git a/Engine/Plugins/Experimental/ControlFlows/Source/ControlFlows/Private/ControlFlow.cpp b/Engine/Plugins/Experimental/ControlFlows/Source/ControlFlows/Private/ControlFlow.cpp
--- a/Engine/Plugins/Experimental/ControlFlows/Source/ControlFlows/Private/ControlFlow.cpp
+++ b/Engine/Plugins/Experimental/ControlFlows/Source/ControlFlows/Private/ControlFlow.cpp
@@ -99,6 +99,9 @@ void FControlFlow::HandleControlFlowNodeCompleted(const TSharedRef<const FContro
 		}
 
 		const bool bCancelRequested = !bInterpretCancelledNodeAsComplete && NodeCompleted->HasCancelBeenRequested();
+		TSharedPtr<FControlFlowNode> PendingCompleteNode = CurrentNode;
+
+		// Reset here to cover all edge cases
 		CurrentNode.Reset();
 		CurrentlyRunningTask.Reset();
 
@@ -124,9 +127,11 @@ void FControlFlow::HandleControlFlowNodeCompleted(const TSharedRef<const FContro
 				}
 				else
 				{
-					bBroadcastingStepComplete = true;
+					// Temporarily re-apply current node so the flow is still considered running and we can inspect it during the step complete delegate
+					CurrentNode = PendingCompleteNode;
 					OnStepCompletedDelegate.Broadcast();
-					bBroadcastingStepComplete = false;
+					CurrentNode.Reset();
+
 					ExecuteNextNodeInQueue();
 				}
 			}
@@ -139,7 +144,11 @@ void FControlFlow::HandleControlFlowNodeCompleted(const TSharedRef<const FContro
 			}
 			else
 			{
+				// Temporarily re-apply current node so the flow is still considered running and we can inspect it during the step complete delegate
+				CurrentNode = PendingCompleteNode;
 				OnStepCompletedDelegate.Broadcast();
+				CurrentNode.Reset();
+
 				OnFlowCompleteDelegate.Broadcast();
 				OnCompleteDelegate_Internal.ExecuteIfBound();				
 			}
@@ -258,7 +267,7 @@ void FControlFlow::Reset()
 
 bool FControlFlow::IsRunning() const
 {
-	return CurrentNode.IsValid() || bBroadcastingStepComplete;
+	return CurrentNode.IsValid();
 }
 
 void FControlFlow::CancelFlow()
diff --git a/Engine/Plugins/Experimental/ControlFlows/Source/ControlFlows/Public/ControlFlow.h b/Engine/Plugins/Experimental/ControlFlows/Source/ControlFlows/Public/ControlFlow.h
--- a/Engine/Plugins/Experimental/ControlFlows/Source/ControlFlows/Public/ControlFlow.h
+++ b/Engine/Plugins/Experimental/ControlFlows/Source/ControlFlows/Public/ControlFlow.h
@@ -393,9 +393,9 @@ private:
 
 	FString DebugName;
 
-	bool bInterpretCancelledNodeAsComplete = false;
+	bool bProfilerEventStarted = false;
 
-	bool bBroadcastingStepComplete = false;
+	bool bInterpretCancelledNodeAsComplete = false;
 
 	int32 UnnamedNodeCounter = 0;
 
@@ -415,8 +415,6 @@ private:
 	TArray<TSharedRef<FControlFlowNode>> FlowQueue;
 
 	TSharedPtr<FTrackedActivity> Activity;
-
-	bool bProfilerEventStarted = false;
 };
 
 #if CPUPROFILERTRACE_ENABLED
