[Backout] - CL50512664
#rnx
Original CL Desc
-----------------------------------------------------------------
[ControlFlow] Make sure current node is valid during step complete notification

#tests Used callback to track which step was active for debugging

#Synced-CL 50546903

diff --git a/Engine/Plugins/Experimental/ControlFlows/Source/ControlFlows/Private/ControlFlow.cpp b/Engine/Plugins/Experimental/ControlFlows/Source/ControlFlows/Private/ControlFlow.cpp
--- a/Engine/Plugins/Experimental/ControlFlows/Source/ControlFlows/Private/ControlFlow.cpp
+++ b/Engine/Plugins/Experimental/ControlFlows/Source/ControlFlows/Private/ControlFlow.cpp
@@ -99,9 +99,6 @@ void FControlFlow::HandleControlFlowNodeCompleted(const TSharedRef<const FContro
 		}
 
 		const bool bCancelRequested = !bInterpretCancelledNodeAsComplete && NodeCompleted->HasCancelBeenRequested();
-		TSharedPtr<FControlFlowNode> PendingCompleteNode = CurrentNode;
-
-		// Reset here to cover all edge cases
 		CurrentNode.Reset();
 		CurrentlyRunningTask.Reset();
 
@@ -127,11 +124,9 @@ void FControlFlow::HandleControlFlowNodeCompleted(const TSharedRef<const FContro
 				}
 				else
 				{
-					// Temporarily re-apply current node so the flow is still considered running and we can inspect it during the step complete delegate
-					CurrentNode = PendingCompleteNode;
+					bBroadcastingStepComplete = true;
 					OnStepCompletedDelegate.Broadcast();
-					CurrentNode.Reset();
-
+					bBroadcastingStepComplete = false;
 					ExecuteNextNodeInQueue();
 				}
 			}
@@ -144,11 +139,7 @@ void FControlFlow::HandleControlFlowNodeCompleted(const TSharedRef<const FContro
 			}
 			else
 			{
-				// Temporarily re-apply current node so the flow is still considered running and we can inspect it during the step complete delegate
-				CurrentNode = PendingCompleteNode;
 				OnStepCompletedDelegate.Broadcast();
-				CurrentNode.Reset();
-
 				OnFlowCompleteDelegate.Broadcast();
 				OnCompleteDelegate_Internal.ExecuteIfBound();				
 			}
@@ -267,7 +258,7 @@ void FControlFlow::Reset()
 
 bool FControlFlow::IsRunning() const
 {
-	return CurrentNode.IsValid();
+	return CurrentNode.IsValid() || bBroadcastingStepComplete;
 }
 
 void FControlFlow::CancelFlow()
diff --git a/Engine/Plugins/Experimental/ControlFlows/Source/ControlFlows/Public/ControlFlow.h b/Engine/Plugins/Experimental/ControlFlows/Source/ControlFlows/Public/ControlFlow.h
--- a/Engine/Plugins/Experimental/ControlFlows/Source/ControlFlows/Public/ControlFlow.h
+++ b/Engine/Plugins/Experimental/ControlFlows/Source/ControlFlows/Public/ControlFlow.h
@@ -393,10 +393,10 @@ private:
 
 	FString DebugName;
 
-	bool bProfilerEventStarted = false;
-
 	bool bInterpretCancelledNodeAsComplete = false;
 
+	bool bBroadcastingStepComplete = false;
+
 	int32 UnnamedNodeCounter = 0;
 
 	int32 UnnamedBranchCounter = 0;
@@ -415,6 +415,8 @@ private:
 	TArray<TSharedRef<FControlFlowNode>> FlowQueue;
 
 	TSharedPtr<FTrackedActivity> Activity;
+
+	bool bProfilerEventStarted = false;
 };
 
 #if CPUPROFILERTRACE_ENABLED
